ElectraPlayer: RELEASE. Moved from NFL to public. Note that this is done via a Delete and "Add New" operation as discussed.

diff --git a/Engine/Plugins/Media/ElectraUtil/Source/ElectraSamples/Private/ElectraTextureSampleAVF.cpp b/Engine/Plugins/Media/ElectraUtil/Source/ElectraSamples/Private/ElectraTextureSampleAVF.cpp
new file mode 100644
--- /dev/null
+++ b/Engine/Plugins/Media/ElectraUtil/Source/ElectraSamples/Private/ElectraTextureSampleAVF.cpp
@@ -0,0 +1,380 @@
+// Copyright Epic Games, Inc. All Rights Reserved.
+
+#include "CoreTypes.h"
+
+#if PLATFORM_MAC || PLATFORM_IOS
+
+#include "ElectraTextureSample.h"
+
+#if WITH_ENGINE
+#include "RenderingThread.h"
+#include "RHI.h"
+#include "RHIStaticStates.h"
+#include "MediaShaders.h"
+#else
+#include "Containers/Array.h"
+#endif
+
+#if WITH_ENGINE
+extern void SafeReleaseMetalObject(id Object);
+#endif
+
+// ------------------------------------------------------------------------------------------------------
+// ------------------------------------------------------------------------------------------------------
+// ------------------------------------------------------------------------------------------------------
+
+FElectraMediaTexConvApple::FElectraMediaTexConvApple()
+{
+#if WITH_ENGINE
+	MetalTextureCache = nullptr;
+#endif
+}
+
+FElectraMediaTexConvApple::~FElectraMediaTexConvApple()
+{
+#if WITH_ENGINE
+	if (MetalTextureCache)
+	{
+		CVMetalTextureCacheRef TextureCacheCopy = MetalTextureCache;
+		SafeReleaseMetalObject((id)TextureCacheCopy);
+	}
+#endif
+}
+
+#if WITH_ENGINE
+namespace
+{
+
+/**
+ * Passes a CV*TextureRef or CVPixelBufferRef through to the RHI to wrap in an RHI texture without traversing system memory.
+ */
+class FTexConvTexResourceWrapper
+	: public FResourceBulkDataInterface
+{
+public:
+
+	FTexConvTexResourceWrapper(CFTypeRef InImageBuffer)
+		: ImageBuffer(InImageBuffer)
+	{
+		check(ImageBuffer);
+		CFRetain(ImageBuffer);
+	}
+
+	virtual ~FTexConvTexResourceWrapper()
+	{
+		CFRelease(ImageBuffer);
+		ImageBuffer = nullptr;
+	}
+
+public:
+
+	//~ FResourceBulkDataInterface interface
+
+	virtual void Discard() override
+	{
+		delete this;
+	}
+
+	virtual const void* GetResourceBulkData() const override
+	{
+		return ImageBuffer;
+	}
+
+	virtual uint32 GetResourceBulkDataSize() const override
+	{
+		return ImageBuffer ? ~0u : 0;
+	}
+
+	virtual EBulkDataType GetResourceType() const override
+	{
+		return EBulkDataType::MediaTexture;
+	}
+
+	CFTypeRef ImageBuffer;
+};
+
+
+/**
+ * Allows for direct GPU mem allocation for texture resource from a CVImageBufferRef's system memory backing store.
+ */
+class FTexConvTexResourceMemory
+	: public FResourceBulkDataInterface
+{
+public:
+	FTexConvTexResourceMemory(CVImageBufferRef InImageBuffer)
+	: ImageBuffer(InImageBuffer)
+	{
+		check(ImageBuffer);
+		CFRetain(ImageBuffer);
+	}
+
+	/**
+	 * @return ptr to the resource memory which has been preallocated
+	 */
+	virtual const void* GetResourceBulkData() const override
+	{
+		CVPixelBufferLockBaseAddress(ImageBuffer, kCVPixelBufferLock_ReadOnly);
+		return CVPixelBufferGetBaseAddress(ImageBuffer);
+	}
+
+	/**
+	 * @return size of resource memory
+	 */
+	virtual uint32 GetResourceBulkDataSize() const override
+	{
+		int32 Pitch = CVPixelBufferGetBytesPerRow(ImageBuffer);
+		int32 Height = CVPixelBufferGetHeight(ImageBuffer);
+		uint32 Size = (Pitch * Height);
+
+		return Size;
+	}
+
+	/**
+	 * Free memory after it has been used to initialize RHI resource
+	 */
+	virtual void Discard() override
+	{
+		CVPixelBufferUnlockBaseAddress(ImageBuffer, kCVPixelBufferLock_ReadOnly);
+		delete this;
+	}
+
+	virtual ~FTexConvTexResourceMemory()
+	{
+		CFRelease(ImageBuffer);
+		ImageBuffer = nullptr;
+	}
+
+	CVImageBufferRef ImageBuffer;
+};
+
+} // namespace anonymous
+
+
+void FElectraMediaTexConvApple::ConvertTexture(FTexture2DRHIRef & InDstTexture, CVImageBufferRef InImageBufferRef)
+{
+	check(IsInRenderingThread());
+
+	const int32 FrameHeight = CVPixelBufferGetHeight(InImageBufferRef);
+	const int32 FrameWidth = CVPixelBufferGetWidth(InImageBufferRef);
+	const int32 FrameStride = CVPixelBufferGetBytesPerRow(InImageBufferRef);
+
+	TRefCountPtr<FRHITexture2D> ShaderResource;
+
+	// We have to support Metal for this object now
+	check(COREVIDEO_SUPPORTS_METAL);
+	check(IsMetalPlatform(GMaxRHIShaderPlatform));
+	{
+		if (!MetalTextureCache)
+		{
+			id<MTLDevice> Device = (id<MTLDevice>)GDynamicRHI->RHIGetNativeDevice();
+			check(Device);
+
+			CVReturn Return = CVMetalTextureCacheCreate(kCFAllocatorDefault, nullptr, Device, nullptr, &MetalTextureCache);
+			check(Return == kCVReturnSuccess);
+		}
+		check(MetalTextureCache);
+
+		if (CVPixelBufferIsPlanar(InImageBufferRef))
+		{
+			FMatrix const* ColorTransform;
+// FIXME: where to get fullrange or not from?
+// FIXME: need to make sure this is properly handling sRGB vs. RGB (as we get a texture from the outside!)
+// NOTE: on all HW we tested to far, we never seem to get here for H.264 material...
+			if(0)//bFullRange)
+			{
+				ColorTransform = &MediaShaders::YuvToRgbRec709Scaled;
+			}
+			else
+			{
+				ColorTransform = &MediaShaders::YuvToRgbRec709Unscaled;
+			}
+
+			// Expecting BiPlanar kCVPixelFormatType_420YpCbCr8BiPlanar Full/Video
+			check(CVPixelBufferGetPlaneCount(InImageBufferRef) == 2);
+
+			int32 YWidth = CVPixelBufferGetWidthOfPlane(InImageBufferRef, 0);
+			int32 YHeight = CVPixelBufferGetHeightOfPlane(InImageBufferRef, 0);
+
+			CVMetalTextureRef YTextureRef = nullptr;
+			CVReturn Result = CVMetalTextureCacheCreateTextureFromImage(kCFAllocatorDefault, MetalTextureCache, InImageBufferRef, nullptr, MTLPixelFormatR8Unorm, YWidth, YHeight, 0, &YTextureRef);
+			check(Result == kCVReturnSuccess);
+			check(YTextureRef);
+
+			int32 UVWidth = CVPixelBufferGetWidthOfPlane(InImageBufferRef, 1);
+			int32 UVHeight = CVPixelBufferGetHeightOfPlane(InImageBufferRef, 1);
+
+			CVMetalTextureRef UVTextureRef = nullptr;
+			Result = CVMetalTextureCacheCreateTextureFromImage(kCFAllocatorDefault, MetalTextureCache, InImageBufferRef, nullptr, MTLPixelFormatRG8Unorm, UVWidth, UVHeight, 1, &UVTextureRef);
+			check(Result == kCVReturnSuccess);
+			check(UVTextureRef);
+
+			// Metal can upload directly from an IOSurface to a 2D texture, so we can just wrap it.
+			FRHIResourceCreateInfo YCreateInfo;
+			YCreateInfo.BulkData = new FTexConvTexResourceWrapper(YTextureRef);
+			YCreateInfo.ResourceArray = nullptr;
+
+			FRHIResourceCreateInfo UVCreateInfo;
+			UVCreateInfo.BulkData = new FTexConvTexResourceWrapper(UVTextureRef);
+			UVCreateInfo.ResourceArray = nullptr;
+
+			ETextureCreateFlags TexCreateFlags = TexCreate_Dynamic | TexCreate_NoTiling | TexCreate_ShaderResource;
+			TRefCountPtr<FRHITexture2D> YTex = RHICreateTexture2D(YWidth, YHeight, PF_G8, 1, 1, TexCreateFlags, YCreateInfo);
+			TRefCountPtr<FRHITexture2D> UVTex = RHICreateTexture2D(UVWidth, UVHeight, PF_R8G8, 1, 1, TexCreateFlags, UVCreateInfo);
+
+			// render video frame into sink texture
+			FRHICommandListImmediate& RHICmdList = FRHICommandListExecutor::GetImmediateCommandList();
+			{
+				// configure media shaders
+				auto ShaderMap = GetGlobalShaderMap(GMaxRHIFeatureLevel);
+				TShaderMapRef<FMediaShadersVS> VertexShader(ShaderMap);
+				TShaderMapRef<FYCbCrConvertPS> PixelShader(ShaderMap);
+
+				FRHIRenderPassInfo RPInfo(InDstTexture, ERenderTargetActions::Load_Store);
+				RHICmdList.BeginRenderPass(RPInfo, TEXT("AvfMediaSampler"));
+				{
+					FGraphicsPipelineStateInitializer GraphicsPSOInit;
+					RHICmdList.ApplyCachedRenderTargets(GraphicsPSOInit);
+
+					GraphicsPSOInit.BlendState = TStaticBlendStateWriteMask<CW_RGBA, CW_NONE, CW_NONE, CW_NONE, CW_NONE, CW_NONE, CW_NONE, CW_NONE>::GetRHI();
+					GraphicsPSOInit.RasterizerState = TStaticRasterizerState<>::GetRHI();
+					GraphicsPSOInit.DepthStencilState = TStaticDepthStencilState<false, CF_Always>::GetRHI();
+
+					GraphicsPSOInit.BoundShaderState.VertexDeclarationRHI = GMediaVertexDeclaration.VertexDeclarationRHI;
+					GraphicsPSOInit.BoundShaderState.VertexShaderRHI = VertexShader.GetVertexShader();
+					GraphicsPSOInit.BoundShaderState.PixelShaderRHI = PixelShader.GetPixelShader();
+					GraphicsPSOInit.PrimitiveType = PT_TriangleStrip;
+
+					SetGraphicsPipelineState(RHICmdList, GraphicsPSOInit);
+
+					PixelShader->SetParameters(RHICmdList, YTex, UVTex, *ColorTransform, MediaShaders::YUVOffset8bits, true);
+
+					FVertexBufferRHIRef VertexBuffer = CreateTempMediaVertexBuffer();
+					RHICmdList.SetStreamSource(0, VertexBuffer, 0);
+					RHICmdList.SetViewport(0, 0, 0.0f, YWidth, YHeight, 1.0f);
+
+					RHICmdList.DrawPrimitive(0, 2, 1);
+				}
+				RHICmdList.EndRenderPass();
+				RHICmdList.CopyToResolveTarget(ShaderResource, ShaderResource, FResolveParams());
+			}
+			CFRelease(YTextureRef);
+			CFRelease(UVTextureRef);
+		}
+		else
+		{
+			//
+			// Grab data directly from image buffer reference
+			// (this will create the output texture here - but will always use the settings here, as data is not converted)
+			//
+			int32 Width = CVPixelBufferGetWidth(InImageBufferRef);
+			int32 Height = CVPixelBufferGetHeight(InImageBufferRef);
+
+			CVMetalTextureRef TextureRef = nullptr;
+			CVReturn Result = CVMetalTextureCacheCreateTextureFromImage(kCFAllocatorDefault, MetalTextureCache, InImageBufferRef, nullptr, MTLPixelFormatBGRA8Unorm_sRGB, Width, Height, 0, &TextureRef);
+			check(Result == kCVReturnSuccess);
+			check(TextureRef);
+
+			FRHIResourceCreateInfo CreateInfo;
+			CreateInfo.BulkData = new FTexConvTexResourceWrapper(TextureRef);
+			CreateInfo.ResourceArray = nullptr;
+
+			ETextureCreateFlags TexCreateFlags = TexCreate_SRGB | TexCreate_Dynamic | TexCreate_NoTiling | TexCreate_ShaderResource;
+			InDstTexture = RHICreateTexture2D(Width, Height, PF_B8G8R8A8, 1, 1, TexCreateFlags, CreateInfo);
+
+			CFRelease(TextureRef);
+		}
+	}
+}
+#endif
+
+// ------------------------------------------------------------------------------------------------------
+// ------------------------------------------------------------------------------------------------------
+// ------------------------------------------------------------------------------------------------------
+
+
+/**
+* Init code for realloacting an image from the the pool
+*/
+void FElectraTextureSample::InitializePoolable()
+{
+}
+
+
+/**
+*  Return the object to the pool and inform the renderer about this...
+*/
+void FElectraTextureSample::ShutdownPoolable()
+{
+	VideoDecoderOutput.Reset();
+}
+
+
+FIntPoint FElectraTextureSample::GetDim() const
+{
+	if (VideoDecoderOutput)
+	{
+		return VideoDecoderOutput->GetDim();
+	}
+	return FIntPoint::ZeroValue;
+}
+
+FMediaTimeStamp FElectraTextureSample::GetTime() const
+{
+	if (VideoDecoderOutput)
+	{
+		const FDecoderTimeStamp TimeStamp = VideoDecoderOutput->GetTime();
+		return FMediaTimeStamp(TimeStamp.Time, TimeStamp.SequenceIndex);
+	}
+	return FMediaTimeStamp();
+}
+
+FTimespan FElectraTextureSample::GetDuration() const
+{
+	if (VideoDecoderOutput)
+	{
+		return VideoDecoderOutput->GetDuration();
+	}
+	return FTimespan::Zero();
+}
+
+FIntPoint FElectraTextureSample::GetOutputDim() const
+{
+	if (VideoDecoderOutput)
+	{
+		return VideoDecoderOutput->GetOutputDim();
+	}
+	return FIntPoint::ZeroValue;
+}
+
+uint32 FElectraTextureSample::GetStride() const
+{
+	if (VideoDecoderOutput)
+	{
+		return VideoDecoderOutput->GetStride();
+	}
+	return 0;
+}
+
+uint32 FElectraTextureSample::GetConverterInfoFlags() const
+{
+	if (VideoDecoderOutput)
+	{
+		return CVPixelBufferIsPlanar(VideoDecoderOutput->GetImageBuffer()) ? ConverterInfoFlags_Default : ConverterInfoFlags_WillCreateOutputTexture;
+	}
+	return ConverterInfoFlags_Default;
+}
+
+bool FElectraTextureSample::Convert(FTexture2DRHIRef & InDstTexture, const FConversionHints & Hints)
+{
+	if (VideoDecoderOutput)
+	{
+		if (TSharedPtr<FElectraMediaTexConvApple, ESPMode::ThreadSafe> PinnedTexConv = TexConv.Pin())
+		{
+			PinnedTexConv->ConvertTexture(InDstTexture, VideoDecoderOutput->GetImageBuffer());
+			return true;
+		}
+	}
+	return false;
+}
+
+#endif
diff --git a/Engine/Plugins/Media/ElectraUtil/Source/ElectraSamples/Private/ElectraTextureSampleAndroid.cpp b/Engine/Plugins/Media/ElectraUtil/Source/ElectraSamples/Private/ElectraTextureSampleAndroid.cpp
new file mode 100644
--- /dev/null
+++ b/Engine/Plugins/Media/ElectraUtil/Source/ElectraSamples/Private/ElectraTextureSampleAndroid.cpp
@@ -0,0 +1,462 @@
+// Copyright Epic Games, Inc. All Rights Reserved.
+
+#include "CoreMinimal.h"
+
+#if PLATFORM_ANDROID
+
+#include "ElectraTextureSample.h"
+
+#include "Android/AndroidPlatform.h"
+#include "Android/AndroidJava.h"
+#include "Android/AndroidApplication.h"
+
+#include "RenderingThread.h"
+
+
+#define ELECTRA_INIT_ON_RENDERTHREAD	1	// set to 1 if context & surface init should be on render thread (seems safer for compatibility - TODO: research why)
+
+// ---------------------------------------------------------------------------------------------------------------------
+
+#if UE_BUILD_SHIPPING
+// always clear any exceptions in shipping
+#define CHECK_JNI_RESULT(Id) if (Id == 0) { JEnv->ExceptionClear(); }
+#else
+#define CHECK_JNI_RESULT(Id)								\
+	if (Id == 0)											\
+	{														\
+		if (bIsOptional)									\
+		{													\
+			JEnv->ExceptionClear();							\
+		}													\
+		else												\
+		{													\
+			JEnv->ExceptionDescribe();						\
+			checkf(Id != 0, TEXT("Failed to find " #Id));	\
+		}													\
+	}
+#endif
+
+// ---------------------------------------------------------------------------------------------------------------------
+
+class FElectraTextureSampleSupport : public FJavaClassObject
+{
+public:
+	FElectraTextureSampleSupport();
+	~FElectraTextureSampleSupport();
+
+	struct FFrameUpdateInfo
+	{
+		int64   Timestamp = 0;
+		int64   Duration = 0;
+		float   UScale = 0.0f;
+		float   UOffset = 0.0f;
+		float   VScale = 0.0f;
+		float   VOffset = 0.0f;
+		bool	bFrameReady = false;
+		bool	bRegionChanged = false;
+		int 	NumPending = 0;
+	};
+
+	int32 GetFrameDataAndUpdateInfo(FFrameUpdateInfo& OutFrameUpdateInfo, FElectraTextureSample* InTargetSample);
+	int32 GetFrameData(FElectraTextureSample* InTargetSample);
+	jobject GetCodecSurface();
+
+private:
+	// Java methods
+	FJavaClassMethod	InitializeFN;
+	FJavaClassMethod	ReleaseFN;
+	FJavaClassMethod	GetCodecSurfaceFN;
+	FJavaClassMethod	GetVideoFrameUpdateInfoFN;
+
+	// FFrameUpdateInfo member field IDs
+	jclass				FFrameUpdateInfoClass;
+	jfieldID			FFrameUpdateInfo_Buffer;
+	jfieldID			FFrameUpdateInfo_Timestamp;
+	jfieldID			FFrameUpdateInfo_Duration;
+	jfieldID			FFrameUpdateInfo_bFrameReady;
+	jfieldID			FFrameUpdateInfo_bRegionChanged;
+	jfieldID			FFrameUpdateInfo_UScale;
+	jfieldID			FFrameUpdateInfo_UOffset;
+	jfieldID			FFrameUpdateInfo_VScale;
+	jfieldID			FFrameUpdateInfo_VOffset;
+	jfieldID			FFrameUpdateInfo_NumPending;
+
+	jobject				CodecSurface;
+	FEvent*				SurfaceInitEvent;
+
+
+	static FName GetClassName()
+	{
+		return FName("com/epicgames/ue4/ElectraTextureSample");
+	}
+
+	jfieldID FindField(JNIEnv* JEnv, jclass InClass, const ANSICHAR* InFieldName, const ANSICHAR* InFieldType, bool bIsOptional)
+	{
+		jfieldID Field = InClass == nullptr ? nullptr : JEnv->GetFieldID(InClass, InFieldName, InFieldType);
+		CHECK_JNI_RESULT(Field);
+		return Field;
+	}
+
+	jmethodID FindMethod(JNIEnv* JEnv, jclass InClass, const ANSICHAR* InMethodName, const ANSICHAR* InMethodSignature, bool bIsOptional)
+	{
+		jmethodID Method = InClass == nullptr ? nullptr : JEnv->GetMethodID(InClass, InMethodName, InMethodSignature);
+		CHECK_JNI_RESULT(Method);
+		return Method;
+	}
+
+	// Create a Java byte array. Must DeleteLocalRef() after use or handling over to Java.
+	jbyteArray MakeJavaByteArray(const uint8* InData, int32 InNumBytes)
+	{
+		JNIEnv* JEnv = FAndroidApplication::GetJavaEnv();
+		jbyteArray RawBuffer = JEnv->NewByteArray(InNumBytes);
+		JEnv->SetByteArrayRegion(RawBuffer, 0, InNumBytes, reinterpret_cast<const jbyte*>(InData));
+		return RawBuffer;
+	}
+
+};
+
+// ---------------------------------------------------------------------------------------------------------------------
+
+FElectraTextureSampleSupport::FElectraTextureSampleSupport()
+	: FJavaClassObject(GetClassName(), "()V")
+	, InitializeFN(GetClassMethod("Initialize", "(Z)V"))
+	, ReleaseFN(GetClassMethod("Release", "()V"))
+	, GetCodecSurfaceFN(GetClassMethod("GetCodecSurface", "()Landroid/view/Surface;"))
+	, GetVideoFrameUpdateInfoFN(GetClassMethod("GetVideoFrameUpdateInfo", "(III)Lcom/epicgames/ue4/ElectraTextureSample$FFrameUpdateInfo;"))
+	, CodecSurface(nullptr)
+	, SurfaceInitEvent(nullptr)
+{
+	JNIEnv* JEnv = FAndroidApplication::GetJavaEnv();
+
+	// Get field IDs for FFrameUpdateInfo class members
+	jclass localFrameUpdateInfoClass = FAndroidApplication::FindJavaClass("com/epicgames/ue4/ElectraTextureSample$FFrameUpdateInfo");
+	FFrameUpdateInfoClass = (jclass)JEnv->NewGlobalRef(localFrameUpdateInfoClass);
+	JEnv->DeleteLocalRef(localFrameUpdateInfoClass);
+	FFrameUpdateInfo_Buffer = FindField(JEnv, FFrameUpdateInfoClass, "Buffer", "Ljava/nio/Buffer;", false);
+	FFrameUpdateInfo_Timestamp = FindField(JEnv, FFrameUpdateInfoClass, "Timestamp", "J", false);
+	FFrameUpdateInfo_Duration = FindField(JEnv, FFrameUpdateInfoClass, "Duration", "J", false);
+	FFrameUpdateInfo_bFrameReady = FindField(JEnv, FFrameUpdateInfoClass, "bFrameReady", "Z", false);
+	FFrameUpdateInfo_bRegionChanged = FindField(JEnv, FFrameUpdateInfoClass, "bRegionChanged", "Z", false);
+	FFrameUpdateInfo_UScale = FindField(JEnv, FFrameUpdateInfoClass, "UScale", "F", false);
+	FFrameUpdateInfo_UOffset = FindField(JEnv, FFrameUpdateInfoClass, "UOffset", "F", false);
+	FFrameUpdateInfo_VScale = FindField(JEnv, FFrameUpdateInfoClass, "VScale", "F", false);
+	FFrameUpdateInfo_VOffset = FindField(JEnv, FFrameUpdateInfoClass, "VOffset", "F", false);
+	FFrameUpdateInfo_NumPending = FindField(JEnv, FFrameUpdateInfoClass, "NumPending", "I", false);
+
+#if ELECTRA_INIT_ON_RENDERTHREAD
+	// enqueue to RT to ensure GL resources are created on the appropriate thread.
+	SurfaceInitEvent = FPlatformProcess::GetSynchEventFromPool(true);
+	ENQUEUE_RENDER_COMMAND(InitElectraTextureSample)([this](FRHICommandListImmediate& RHICmdList)
+		{
+			RHICmdList.EnqueueLambda([this](FRHICommandListImmediate& CmdList)
+				{
+					JNIEnv* JEnv = FAndroidApplication::GetJavaEnv();
+					// Setup Java side of things
+					JEnv->CallVoidMethod(Object, InitializeFN.Method, FAndroidMisc::ShouldUseVulkan());
+					// Query surface to be used for decoder
+					jobject Surface = JEnv->CallObjectMethod(Object, GetCodecSurfaceFN.Method);
+					CodecSurface = JEnv->NewGlobalRef(Surface);
+					JEnv->DeleteLocalRef(Surface);
+					SurfaceInitEvent->Trigger();
+				});
+		});
+	FlushRenderingCommands();
+#else
+	// Setup Java side of things
+	JEnv->CallVoidMethod(Object, InitializeFN.Method, FAndroidMisc::ShouldUseVulkan());
+
+	// Query surface to be used for decoder
+	jobject Surface = JEnv->CallObjectMethod(Object, GetCodecSurfaceFN.Method);
+	CodecSurface = JEnv->NewGlobalRef(Surface);
+	JEnv->DeleteLocalRef(Surface);
+#endif
+}
+
+
+FElectraTextureSampleSupport::~FElectraTextureSampleSupport()
+{
+	// When initialization of the surface was triggered on the rendering thread we need to wait for its completion.
+	if (SurfaceInitEvent)
+	{
+		// Wait for the surface initialization event to have been signaled.
+		// Do not wait if we are on the render thread. In this case the initialization has already completed anyway.
+		if (!IsInRenderingThread())
+		{
+			SurfaceInitEvent->Wait();
+		}
+		FPlatformProcess::ReturnSynchEventToPool(SurfaceInitEvent);
+		SurfaceInitEvent = nullptr;
+	}
+
+	if (IsInGameThread())
+	{
+		// enqueue to RT to ensure GL resources are released on the appropriate thread.
+		ENQUEUE_RENDER_COMMAND(DestroyElectraTextureSample)([this](FRHICommandListImmediate& RHICmdList)
+		{
+			RHICmdList.EnqueueLambda([this](FRHICommandListImmediate& CmdList)
+			{
+				JNIEnv* JEnv = FAndroidApplication::GetJavaEnv();
+				JEnv->CallVoidMethod(Object, ReleaseFN.Method);
+				JEnv->DeleteGlobalRef(CodecSurface);
+				JEnv->DeleteGlobalRef(FFrameUpdateInfoClass);
+				CodecSurface = 0;
+				FFrameUpdateInfoClass = 0;
+			});
+		});
+		FlushRenderingCommands();
+	}
+	else
+	{
+		JNIEnv* JEnv = FAndroidApplication::GetJavaEnv();
+		JEnv->CallVoidMethod(Object, ReleaseFN.Method);
+		JEnv->DeleteGlobalRef(CodecSurface);
+		JEnv->DeleteGlobalRef(FFrameUpdateInfoClass);
+		CodecSurface = 0;
+		FFrameUpdateInfoClass = 0;
+	}
+}
+
+//-----------------------------------------------------------------------------
+/**
+ *
+ * @note Call this from an RHI thread! It will need a valid rendering environment!
+ */
+int32 FElectraTextureSampleSupport::GetFrameDataAndUpdateInfo(FFrameUpdateInfo& OutFrameUpdateInfo, FElectraTextureSample* InTargetSample)
+{
+	// In case this is called with a ES renderer, we need to pass in the destination texture we'd like to be used to receive the data
+	// (for Vulkan we'll just receive a simple byte buffer)
+	int32 DestTexture = 0;
+	if (FRHITexture* Texture = InTargetSample->GetTexture())
+	{
+		DestTexture = *reinterpret_cast<int32*>(Texture->GetNativeResource());
+	}
+
+	// Update frame info and get data...
+	JNIEnv* JEnv = FAndroidApplication::GetJavaEnv();
+	jobject OutputInfo = JEnv->CallObjectMethod(Object, GetVideoFrameUpdateInfoFN.Method, DestTexture, InTargetSample->GetDim().X, InTargetSample->GetDim().Y);
+	if (JEnv->ExceptionCheck())
+	{
+		JEnv->ExceptionDescribe();
+		JEnv->ExceptionClear();
+	}
+	// Failure will return no object.
+	if (OutputInfo != nullptr)
+	{
+		OutFrameUpdateInfo.Timestamp = JEnv->GetLongField(OutputInfo, FFrameUpdateInfo_Timestamp);
+		OutFrameUpdateInfo.Duration = JEnv->GetLongField(OutputInfo, FFrameUpdateInfo_Duration);
+		OutFrameUpdateInfo.UScale = JEnv->GetFloatField(OutputInfo, FFrameUpdateInfo_UScale);
+		OutFrameUpdateInfo.UOffset = JEnv->GetFloatField(OutputInfo, FFrameUpdateInfo_UOffset);
+		OutFrameUpdateInfo.VScale = JEnv->GetFloatField(OutputInfo, FFrameUpdateInfo_VScale);
+		OutFrameUpdateInfo.VOffset = JEnv->GetFloatField(OutputInfo, FFrameUpdateInfo_VOffset);
+		OutFrameUpdateInfo.bFrameReady = JEnv->GetBooleanField(OutputInfo, FFrameUpdateInfo_bFrameReady);
+		OutFrameUpdateInfo.bRegionChanged = JEnv->GetBooleanField(OutputInfo, FFrameUpdateInfo_bRegionChanged);
+		OutFrameUpdateInfo.NumPending = JEnv->GetIntField(OutputInfo, FFrameUpdateInfo_NumPending);
+		if (InTargetSample)
+		{
+			jobject buffer = JEnv->GetObjectField(OutputInfo, FFrameUpdateInfo_Buffer);
+			if (buffer != nullptr)
+			{
+				const void* outPixels = JEnv->GetDirectBufferAddress(buffer);
+				const int32 outCount = JEnv->GetDirectBufferCapacity(buffer);
+				InTargetSample->SetupFromBuffer(outPixels, outCount);
+				JEnv->DeleteLocalRef(buffer);
+			}
+		}
+
+		JEnv->DeleteLocalRef(OutputInfo);
+		return 0;
+	}
+	return 1;
+}
+
+//-----------------------------------------------------------------------------
+/**
+ *
+ * @note Call this from an RHI thread! It will need a valid rendering environment!
+ */
+int32 FElectraTextureSampleSupport::GetFrameData(FElectraTextureSample* InTargetSample)
+{
+	// In case this is called with a ES renderer, we need to pass in the destination texture we'd like to be used to receive the data
+	// (for Vulkan we'll just receive a simple byte buffer)
+	int32 DestTexture = 0;
+	if (FRHITexture* Texture = InTargetSample->GetTexture())
+	{
+		DestTexture = *reinterpret_cast<int32*>(Texture->GetNativeResource());
+	}
+
+	// Update frame info and get data...
+	JNIEnv* JEnv = FAndroidApplication::GetJavaEnv();
+	jobject OutputInfo = JEnv->CallObjectMethod(Object, GetVideoFrameUpdateInfoFN.Method, DestTexture, InTargetSample->GetDim().X, InTargetSample->GetDim().Y);
+	if (JEnv->ExceptionCheck())
+	{
+		JEnv->ExceptionDescribe();
+		JEnv->ExceptionClear();
+	}
+	// Failure will return no object.
+	if (OutputInfo != nullptr)
+	{
+		if (InTargetSample)
+		{
+			jobject buffer = JEnv->GetObjectField(OutputInfo, FFrameUpdateInfo_Buffer);
+			if (buffer != nullptr)
+			{
+				const void* outPixels = JEnv->GetDirectBufferAddress(buffer);
+				const int32 outCount = JEnv->GetDirectBufferCapacity(buffer);
+				InTargetSample->SetupFromBuffer(outPixels, outCount);
+				JEnv->DeleteLocalRef(buffer);
+			}
+		}
+
+		JEnv->DeleteLocalRef(OutputInfo);
+		return 0;
+	}
+	return 1;
+}
+
+
+jobject FElectraTextureSampleSupport::GetCodecSurface()
+{
+	if (SurfaceInitEvent)
+	{
+		// Wait for the surface initialization event to have been signaled.
+		// Do not wait if we are on the render thread. In this case the initialization has already completed anyway.
+		if (!IsInRenderingThread())
+		{
+			// Only wait for a little while here just in case this would prevent the render thread from
+			// even starting its jobs and us causing a deadlock here.
+			bool bInitDone = SurfaceInitEvent->Wait(FTimespan::FromMilliseconds(2000.0));
+			if (bInitDone)
+			{
+				// When init has completed we can free the event and do not have to wait for it any more in the future.
+				FPlatformProcess::ReturnSynchEventToPool(SurfaceInitEvent);
+				SurfaceInitEvent = nullptr;
+			}
+		}
+	}
+
+	return CodecSurface;
+}
+
+//-----------------------------------------------------------------------------
+//-----------------------------------------------------------------------------
+//-----------------------------------------------------------------------------
+
+FElectraTextureSamplePool::FElectraTextureSamplePool()
+	: TMediaObjectPool<TextureSample, FElectraTextureSamplePool>(this)
+	, Support(new FElectraTextureSampleSupport())
+{}
+
+
+void* FElectraTextureSamplePool::GetCodecSurface()
+{
+	return (void*)Support->GetCodecSurface();
+}
+
+//-----------------------------------------------------------------------------
+//-----------------------------------------------------------------------------
+//-----------------------------------------------------------------------------
+
+/**
+* Init code for reallocating an image from the the pool
+*/
+void FElectraTextureSample::InitializePoolable()
+{
+}
+
+/**
+*  Return the object to the pool and inform the renderer about this...
+*/
+void FElectraTextureSample::ShutdownPoolable()
+{
+	VideoDecoderOutput.Reset();
+}
+
+
+void FElectraTextureSample::Initialize(FVideoDecoderOutput* InVideoDecoderOutput)
+{
+	VideoDecoderOutput = StaticCastSharedPtr<FVideoDecoderOutputAndroid, IDecoderOutputPoolable, ESPMode::ThreadSafe>(InVideoDecoderOutput->AsShared());
+
+	if (VideoDecoderOutput->GetOutputType() == FVideoDecoderOutputAndroid::EOutputType::DirectToSurfaceAsQueue)
+	{
+		ENQUEUE_RENDER_COMMAND(InitTextureSample)([WeakThis{ TWeakPtr<FElectraTextureSample, ESPMode::ThreadSafe>(AsShared()) }](FRHICommandListImmediate& RHICmdList) {
+			if (TSharedPtr<FElectraTextureSample, ESPMode::ThreadSafe> This = WeakThis.Pin())
+			{
+				This->InitializeTexture();
+
+				if (This->Texture)
+				{
+					RHICmdList.EnqueueLambda([WeakThis](FRHICommandListImmediate& CmdList)
+					{
+						QUICK_SCOPE_CYCLE_COUNTER(STAT_FRHICommandUpdateDecoderExternaTexture_Execute);
+						if (TSharedPtr<FElectraTextureSample, ESPMode::ThreadSafe> This = WeakThis.Pin())
+						{
+							This->Support->GetFrameData(This.Get());
+						}
+					});
+				}
+				else
+				{
+					// CPU side buffer case: we do not need to do that on an RHI thread
+					This->Support->GetFrameData(This.Get());
+				}
+			}
+		});
+	}
+}
+
+
+void FElectraTextureSample::InitializeTexture()
+{
+	check(IsInRenderingThread() || IsInRHIThread());
+
+	if (FAndroidMisc::ShouldUseVulkan())
+	{
+		Texture = nullptr;
+		return;
+	}
+
+	FIntPoint Dim = VideoDecoderOutput->GetDim();
+
+	if (Texture.IsValid() && (Texture->GetSizeXY() == Dim))
+	{
+		return;
+	}
+
+	TRefCountPtr<FRHITexture2D> DummyTexture2DRHI;
+	FRHIResourceCreateInfo CreateInfo;
+
+	RHICreateTargetableShaderResource2D(
+		Dim.X,
+		Dim.Y,
+		PF_B8G8R8A8,
+		1,
+		TexCreate_Dynamic | TexCreate_SRGB,
+		TexCreate_RenderTargetable,
+		false,
+		CreateInfo,
+		Texture,
+		DummyTexture2DRHI
+	);
+
+	return;
+}
+
+
+void FElectraTextureSample::SetupFromBuffer(const void* InBuffer, int32 InBufferSize)
+{
+	if (BufferSize < InBufferSize)
+	{
+		if (BufferSize == 0)
+		{
+			Buffer = FMemory::Malloc(InBufferSize);
+		}
+		else
+		{
+			Buffer = FMemory::Realloc(Buffer, InBufferSize);
+		}
+		BufferSize = InBufferSize;
+	}
+	FMemory::Memcpy(Buffer, InBuffer, InBufferSize);
+}
+
+#endif
diff --git a/Engine/Plugins/Media/ElectraUtil/Source/ElectraSamples/Private/ElectraTextureSamplePC.cpp b/Engine/Plugins/Media/ElectraUtil/Source/ElectraSamples/Private/ElectraTextureSamplePC.cpp
new file mode 100644
--- /dev/null
+++ b/Engine/Plugins/Media/ElectraUtil/Source/ElectraSamples/Private/ElectraTextureSamplePC.cpp
@@ -0,0 +1,302 @@
+// Copyright Epic Games, Inc. All Rights Reserved.
+
+#include "CoreMinimal.h"
+
+#if PLATFORM_WINDOWS && !UE_SERVER
+
+#include "ElectraTextureSample.h"
+
+#include "ProfilingDebugging/RealtimeGPUProfiler.h"
+#include "RenderUtils.h"
+
+#ifdef ELECTRA_HAVE_DX11
+#include "Windows/AllowWindowsPlatformTypes.h"
+#include "D3D11State.h"
+#include "D3D11Resources.h"
+#include "Windows/HideWindowsPlatformTypes.h"
+#endif
+
+/*
+	Short summary of how we get data:
+
+	- Win10+ (HW decode is used at all times)
+	-- DX11: we receive data in GPU space as NV12 texture
+	-- DX12: we receive data in CPU(yes) space as NV12 texture
+
+	- Win8:
+	-- SW-decode fallback: we receive data in a shared DX11 texture (despite it being SW decode) in NV12 format
+
+	- Win7:
+	-- we receive data in a CPU space buffer in NV12 format
+*/
+
+// -------------------------------------------------------------------------------------------------------------------------
+
+DECLARE_GPU_STAT_NAMED(MediaWinDecoder_Convert, TEXT("MediaWinDecoder_Convert"));
+
+// --------------------------------------------------------------------------------------------------------------------------
+
+void FElectraTextureSample::Initialize(FVideoDecoderOutput *InVideoDecoderOutput)
+{
+	VideoDecoderOutput = StaticCastSharedPtr<FVideoDecoderOutputPC, IDecoderOutputPoolable, ESPMode::ThreadSafe>(InVideoDecoderOutput->AsShared());
+	SampleFormat = EMediaTextureSampleFormat::CharNV12;
+}
+
+
+#if !UE_SERVER
+void FElectraTextureSample::InitializePoolable()
+{
+}
+
+void FElectraTextureSample::ShutdownPoolable()
+{
+	VideoDecoderOutput.Reset();
+}
+#endif
+
+IMediaTextureSampleConverter* FElectraTextureSample::GetMediaTextureSampleConverter()
+{
+#ifdef ELECTRA_HAVE_DX11
+	// All versions might need SW fallback - check if we have a real texture as source -> converter needed
+	return (VideoDecoderOutput && VideoDecoderOutput->GetTexture()) ? this : nullptr;
+#else
+	return nullptr;
+#endif
+}
+
+#ifdef ELECTRA_HAVE_DX11
+struct FRHICommandCopyResource final : public FRHICommand<FRHICommandCopyResource>
+{
+	TRefCountPtr<ID3D11Texture2D> SampleTexture;
+	FTexture2DRHIRef SampleDestinationTexture;
+	bool bCrossDevice;
+
+	FRHICommandCopyResource(ID3D11Texture2D* InSampleTexture, FRHITexture2D* InSampleDestinationTexture, bool bInCrossDevice)
+		: SampleTexture(InSampleTexture)
+		, SampleDestinationTexture(InSampleDestinationTexture)
+		, bCrossDevice(bInCrossDevice)
+	{
+	}
+
+	void Execute(FRHICommandListBase& CmdList)
+	{
+		LLM_SCOPE(ELLMTag::MediaStreaming);
+		ID3D11Device* D3D11Device = static_cast<ID3D11Device*>(GDynamicRHI->RHIGetNativeDevice());
+		ID3D11DeviceContext* D3D11DeviceContext = nullptr;
+
+		D3D11Device->GetImmediateContext(&D3D11DeviceContext);
+		if (D3D11DeviceContext)
+		{
+			ID3D11Resource* DestinationTexture = reinterpret_cast<ID3D11Resource*>(SampleDestinationTexture->GetNativeResource());
+			if (DestinationTexture)
+			{
+				if (bCrossDevice)
+				{
+					TRefCountPtr<IDXGIResource> OtherResource(nullptr);
+					SampleTexture->QueryInterface(__uuidof(IDXGIResource), (void**)OtherResource.GetInitReference());
+
+					if (OtherResource)
+					{
+						//
+						// Copy shared texture from decoder device to render device
+						//
+						HANDLE SharedHandle = nullptr;
+						if (OtherResource->GetSharedHandle(&SharedHandle) == S_OK)
+						{
+							if (SharedHandle != 0)
+							{
+								TRefCountPtr<ID3D11Resource> SharedResource;
+								D3D11Device->OpenSharedResource(SharedHandle, __uuidof(ID3D11Texture2D), (void**)SharedResource.GetInitReference());
+
+								if (SharedResource)
+								{
+									TRefCountPtr<IDXGIKeyedMutex> KeyedMutex;
+									SharedResource->QueryInterface(_uuidof(IDXGIKeyedMutex), (void**)KeyedMutex.GetInitReference());
+
+									if (KeyedMutex)
+									{
+										// Key is 1 : Texture as just been updated
+										// Key is 2 : Texture as already been updated.
+										// Do not wait to acquire key 1 since there is race no condition between writer and reader.
+										if (KeyedMutex->AcquireSync(1, 0) == S_OK)
+										{
+											// Copy from shared texture of FSink device to Rendering device
+											D3D11DeviceContext->CopyResource(DestinationTexture, SharedResource);
+											KeyedMutex->ReleaseSync(2);
+										}
+										else
+										{
+											// If key 1 cannot be acquired, another reader is already copying the resource
+											// and will release key with 2. 
+											// Wait to acquire key 2.
+											if (KeyedMutex->AcquireSync(2, INFINITE) == S_OK)
+											{
+												KeyedMutex->ReleaseSync(2);
+											}
+										}
+									}
+								}
+							}
+						}
+					}
+				}
+				else
+				{
+					//
+					// Simple copy on render device
+					//
+					D3D11DeviceContext->CopyResource(DestinationTexture, SampleTexture);
+				}
+			}
+			D3D11DeviceContext->Release();
+		}
+	}
+};
+#endif
+
+/**
+ * "Converter" for textures - here: a copy from the decoder owned texture (possibly in another device) into a RHI one (as prep for the real conversion to RGB etc.)
+ */
+bool FElectraTextureSample::Convert(FTexture2DRHIRef & InDstTexture, const FConversionHints & Hints)
+{
+#ifdef ELECTRA_HAVE_DX11
+
+	LLM_SCOPE(ELLMTag::MediaStreaming);
+
+	check(IsInRenderingThread());
+
+	FRHICommandListImmediate& RHICmdList = FRHICommandListExecutor::GetImmediateCommandList();
+
+	SCOPED_DRAW_EVENT(RHICmdList, WinMediaOutputConvertTexture);
+	SCOPED_GPU_STAT(RHICmdList, MediaWinDecoder_Convert);
+
+	// Get actual sample dimensions
+	FIntPoint Dim = VideoDecoderOutput->GetDim();
+
+	bool bCrossDeviceCopy;
+	uint8 Format;
+	if (VideoDecoderOutput->GetOutputType() != FVideoDecoderOutputPC::EOutputType::HardwareWin8Plus)
+	{
+		//
+		// SW decoder has decoded into a HW texture (not known to RHI) -> copy it into an RHI one
+		//
+		check(VideoDecoderOutput->GetOutputType() == FVideoDecoderOutputPC::EOutputType::SoftwareWin8Plus);
+		check(VideoDecoderOutput->GetFormat() == EPixelFormat::PF_NV12);
+
+		Format = EPixelFormat::PF_G8; // use fixed format: we flag this as NV12, too - as it is - but DX11 will only support a "higher than normal G8 texture" (with specialized access in shader)
+		bCrossDeviceCopy = false;
+	}
+	else
+	{
+		//
+		// HW decoder has delivered a texture (this is already a copy) which is on its own device. Copy into one created by RHI (and hence on our rendering device)
+		//
+
+		// note: on DX platforms we won't get any SRV generated for NV12 -> so any user needs to do that manually! (as they please: R8, R8G8...)
+		Format = VideoDecoderOutput->GetFormat();
+		bCrossDeviceCopy = true;
+	}
+
+	// Do we need a new RHI texture?
+	if (!Texture.IsValid() || Texture->GetSizeX() != Dim.X || Texture->GetSizeY() != Dim.Y)
+	{
+		FRHIResourceCreateInfo CreateInfo;
+		const ETextureCreateFlags CreateFlags = TexCreate_Dynamic;
+		Texture = RHICreateTexture2D(
+			Dim.X,
+			Dim.Y,
+			Format,
+			1,
+			1,
+			CreateFlags,
+			CreateInfo);
+	}
+	
+	// Copy data into RHI texture
+	if (RHICmdList.Bypass())
+	{
+		FRHICommandCopyResource Cmd(VideoDecoderOutput->GetTexture(), Texture, bCrossDeviceCopy);
+		Cmd.Execute(RHICmdList);
+	}
+	else
+	{
+		new (RHICmdList.AllocCommand<FRHICommandCopyResource>()) FRHICommandCopyResource(VideoDecoderOutput->GetTexture(), Texture, bCrossDeviceCopy);
+	}
+	
+	return true;
+#else
+	return false;
+#endif
+}
+
+
+const void* FElectraTextureSample::GetBuffer()
+{
+	if (VideoDecoderOutput)
+	{
+		return VideoDecoderOutput->GetBuffer().GetData();
+	}
+	return nullptr;
+}
+
+
+uint32 FElectraTextureSample::GetStride() const
+{
+	if (VideoDecoderOutput)
+	{
+		return VideoDecoderOutput->GetStride();
+	}
+	return 0;
+}
+
+
+FIntPoint FElectraTextureSample::GetDim() const
+{
+	if (VideoDecoderOutput)
+	{
+		return VideoDecoderOutput->GetDim();
+	}
+	return FIntPoint::ZeroValue;
+}
+
+
+FIntPoint FElectraTextureSample::GetOutputDim() const
+{
+	if (VideoDecoderOutput)
+	{
+		return VideoDecoderOutput->GetOutputDim();
+	}
+	return FIntPoint::ZeroValue;
+}
+
+
+FMediaTimeStamp FElectraTextureSample::GetTime() const
+{
+	if (VideoDecoderOutput)
+	{
+		const FDecoderTimeStamp TimeStamp = VideoDecoderOutput->GetTime();
+		return FMediaTimeStamp(TimeStamp.Time, TimeStamp.SequenceIndex);
+	}
+	return FMediaTimeStamp();
+}
+
+
+FTimespan FElectraTextureSample::GetDuration() const
+{
+	if (VideoDecoderOutput)
+	{
+		return VideoDecoderOutput->GetDuration();
+	}
+	return FTimespan::Zero();
+}
+
+IMFSample* FElectraTextureSample::GetMFSample()
+{
+	if (VideoDecoderOutput)
+	{
+		return VideoDecoderOutput->GetMFSample().GetReference();
+	}
+	return nullptr;
+}
+
+#endif
diff --git a/Engine/Plugins/Media/ElectraUtil/Source/ElectraSamples/Private/ElectraTextureSampleUnix.cpp b/Engine/Plugins/Media/ElectraUtil/Source/ElectraSamples/Private/ElectraTextureSampleUnix.cpp
new file mode 100644
--- /dev/null
+++ b/Engine/Plugins/Media/ElectraUtil/Source/ElectraSamples/Private/ElectraTextureSampleUnix.cpp
@@ -0,0 +1,32 @@
+// Copyright Epic Games, Inc. All Rights Reserved.
+
+#include "CoreMinimal.h"
+
+#if PLATFORM_UNIX
+
+#include "ElectraTextureSample.h"
+
+/**
+* Init code for realloacting an image from the the pool
+*/
+void FElectraTextureSampleUnix::InitializePoolable()
+{
+	Duration = FTimespan::Zero();
+}
+
+/**
+*  Return the object to the pool and inform the renderer about this...
+*/
+void FElectraTextureSampleUnix::ShutdownPoolable()
+{
+	TSharedPtr<MEDIArendererVideoUE, ESPMode::ThreadSafe> lockedVideoRenderer = OwningRenderer.Pin();
+	if (lockedVideoRenderer.IsValid())
+	{
+		lockedVideoRenderer->TextureReleasedToPool(GetDuration());
+	}
+
+	// This is coming from the original code: Drop reference to the texture. It should be released by the outside system.
+	Texture = nullptr;
+}
+
+#endif
\ No newline at end of file
diff --git a/Engine/Plugins/Media/ElectraUtil/Source/ElectraSamples/Private/Java/src/com/epicgames/ue4/ElectraTextureSample.java b/Engine/Plugins/Media/ElectraUtil/Source/ElectraSamples/Private/Java/src/com/epicgames/ue4/ElectraTextureSample.java
new file mode 100644
--- /dev/null
+++ b/Engine/Plugins/Media/ElectraUtil/Source/ElectraSamples/Private/Java/src/com/epicgames/ue4/ElectraTextureSample.java
@@ -0,0 +1,939 @@
+// Copyright Epic Games, Inc. All Rights Reserved.
+
+package com.epicgames.ue4;
+
+import android.os.Build;
+import java.io.IOException;
+import java.util.ArrayList;
+import android.util.Log;
+import android.util.SparseArray;
+
+import android.opengl.*;
+import android.graphics.SurfaceTexture;
+import android.view.Surface;
+import java.nio.ByteBuffer;
+
+
+public class ElectraTextureSample
+{
+	private static final String TAG = "[ElectraPlayerSDK]";
+
+	public class FFrameUpdateInfo
+	{
+		public java.nio.Buffer Buffer;
+		public long Timestamp = 0;
+		public long Duration = 0;
+		public boolean bFrameReady = false;
+		public boolean bRegionChanged = false;
+		public float UScale = 0.0f;
+		public float UOffset = 0.0f;
+		public float VScale = 0.0f;
+		public float VOffset = 0.0f;
+		public int NumPending = 0;
+	}
+
+	private FBitmapRenderer BitmapRenderer = null;
+
+	public ElectraTextureSample()
+	{
+	}
+
+	public void Initialize(boolean bVulkan)
+	{
+		BitmapRenderer = new FBitmapRenderer(bVulkan);
+	}
+
+	public void Release()
+	{
+		if (BitmapRenderer != null)
+		{
+			BitmapRenderer.release();
+			BitmapRenderer = null;
+		}
+	}
+
+	public FFrameUpdateInfo GetVideoFrameUpdateInfo(int destTextureES, int width, int height)
+	{
+		FFrameUpdateInfo Info = null;
+		if (BitmapRenderer != null)
+		{
+			Info = BitmapRenderer.GetVideoFrameUpdateInfo(destTextureES, width, height);
+		}
+		return Info;
+	}
+
+	public android.view.Surface GetCodecSurface()
+	{
+		if (BitmapRenderer != null)
+		{
+			return BitmapRenderer.getSurface();
+		}
+		return null;
+	}
+
+	/***************************************************************************************************************************************************/
+	/***************************************************************************************************************************************************/
+	/***************************************************************************************************************************************************/
+
+	/*
+		All this internal surface view does is manage the offscreen bitmap that the media player decoding can render into for eventual extraction to the UE4 buffers.
+	*/
+	class FBitmapRenderer
+	{
+		private java.nio.Buffer mFrameData = null;
+		private int mFrameDataSize = 0;
+		private android.graphics.SurfaceTexture mSurfaceTexture = null;
+		private int mTextureWidth = -1;
+		private int mTextureHeight = -1;
+		private android.view.Surface mSurface = null;
+		private int mTextureID = -1;
+		private int mFBO = -1;
+		private int mBlitVertexShaderID = -1;
+		private int mBlitFragmentShaderID = -1;
+		private float[] mTransformMatrix = new float[16];
+		private boolean mTriangleVerticesDirty = true;
+		private boolean mUseOwnContext = true;
+		private boolean mVulkanRenderer = false;
+		private boolean mSwizzlePixels = false;
+
+		private int GL_TEXTURE_EXTERNAL_OES = 0x8D65;
+
+		private EGLDisplay mEglDisplay;
+		private EGLContext mEglContext;
+		private EGLSurface mEglSurface;
+
+		private EGLDisplay mSavedDisplay;
+		private EGLContext mSavedContext;
+		private EGLSurface mSavedSurfaceDraw;
+		private EGLSurface mSavedSurfaceRead;
+
+		private boolean mCreatedEGLDisplay = false;
+
+		public int mNativeDecoderID;
+
+		public FBitmapRenderer(boolean vulkanRenderer)
+		{
+			mVulkanRenderer = vulkanRenderer;
+
+			mEglSurface = EGL14.EGL_NO_SURFACE;
+			mEglContext = EGL14.EGL_NO_CONTEXT;
+			mUseOwnContext = true;
+
+			if (mVulkanRenderer)
+			{
+				mSwizzlePixels = true;
+
+				Log.d(TAG,"FBitmapRenderer: mSwizzlePixels");
+			}
+			else
+			{
+				String RendererString = GLES20.glGetString(GLES20.GL_RENDERER);
+				Log.d(TAG,"FBitmapRenderer: GLES: " + RendererString);
+				// Do not use shared context if Adreno before 400 or on older Android than Marshmallow
+				if (RendererString.contains("Adreno (TM) "))
+				{
+					int AdrenoVersion = Integer.parseInt(RendererString.substring(12));
+					if (AdrenoVersion < 400 || android.os.Build.VERSION.SDK_INT < 22)
+					{
+						Log.d(TAG,"VideoDecoder: disabled shared GL context on " + RendererString);
+						mUseOwnContext = false;
+					}
+				}
+			}
+
+			if (mUseOwnContext)
+			{
+				initContext();
+				saveContext();
+				makeCurrent();
+				initSurfaceTexture();
+				restoreContext();
+			}
+			else
+			{
+				initSurfaceTexture();
+			}
+		}
+
+		private void initContext()
+		{
+			mEglDisplay = EGL14.EGL_NO_DISPLAY;
+			EGLContext shareContext = EGL14.EGL_NO_CONTEXT;
+
+			int majorver[] = new int[] { 0 };
+			int minorver[] = new int[] { 0 };
+			if (!mVulkanRenderer)
+			{
+				mEglDisplay = EGL14.eglGetCurrentDisplay();
+				shareContext = EGL14.eglGetCurrentContext();
+
+				if (android.os.Build.VERSION.SDK_INT >= 18 &&
+					EGL14.eglQueryContext(mEglDisplay, shareContext, EGLExt.EGL_CONTEXT_MAJOR_VERSION_KHR, majorver, 0) &&
+					EGL14.eglQueryContext(mEglDisplay, shareContext, EGLExt.EGL_CONTEXT_MINOR_VERSION_KHR, minorver, 0))
+				{
+					Log.d(TAG,"VideoDecoder: Existing GL context is version " + majorver[0] + "." + minorver[0]);
+				}
+				else
+				// on some devices eg Galaxy S6, the above fails but we do get EGL14.EGL_CONTEXT_CLIENT_VERSION=3
+				if (EGL14.eglQueryContext(mEglDisplay, shareContext, EGL14.EGL_CONTEXT_CLIENT_VERSION, majorver, 0))
+				{
+					Log.d(TAG, "VideoDecoder: Existing GL context is version " + majorver[0]);
+				}
+				else
+				{
+					Log.d(TAG, "VideoDecoder: Existing GL context version not detected");
+				}
+			}
+			else
+			{
+				mEglDisplay = EGL14.eglGetDisplay(EGL14.EGL_DEFAULT_DISPLAY);
+				if (mEglDisplay == EGL14.EGL_NO_DISPLAY)
+				{
+					Log.e(TAG, "unable to get EGL14 display");
+					return;
+				}
+				int[] version = new int[2];
+				if (!EGL14.eglInitialize(mEglDisplay, version, 0, version, 1))
+				{
+					mEglDisplay = null;
+					Log.e(TAG,"unable to initialize EGL14 display");
+					return;
+				}
+
+				mCreatedEGLDisplay = true;
+			}
+
+			int[] configSpec = new int[]
+			{
+				EGL14.EGL_RENDERABLE_TYPE, EGL14.EGL_OPENGL_ES2_BIT,
+				EGL14.EGL_SURFACE_TYPE, EGL14.EGL_PBUFFER_BIT,
+				EGL14.EGL_NONE
+			};
+			EGLConfig[] configs = new EGLConfig[1];
+			int[] num_config = new int[1];
+			EGL14.eglChooseConfig(mEglDisplay, configSpec, 0, configs, 0, 1, num_config, 0);
+			int[] contextAttribsES2 = new int[]
+			{
+				EGL14.EGL_CONTEXT_CLIENT_VERSION, 2,
+				EGL14.EGL_NONE
+			};
+			int[] contextAttribsES31 = new int[]
+			{
+				EGLExt.EGL_CONTEXT_MAJOR_VERSION_KHR, 3,
+				EGLExt.EGL_CONTEXT_MINOR_VERSION_KHR, 1,
+				EGL14.EGL_NONE
+			};
+			mEglContext = EGL14.eglCreateContext(mEglDisplay, configs[0], shareContext, majorver[0]==3 ? contextAttribsES31 : contextAttribsES2, 0);
+
+			if (EGL14.eglQueryString(mEglDisplay, EGL14.EGL_EXTENSIONS).contains("EGL_KHR_surfaceless_context"))
+			{
+				mEglSurface = EGL14.EGL_NO_SURFACE;
+			}
+			else
+			{
+				int[] pbufferAttribs = new int[]
+				{
+					EGL14.EGL_NONE
+				};
+				mEglSurface = EGL14.eglCreatePbufferSurface(mEglDisplay, configs[0], pbufferAttribs, 0);
+			}
+		}
+
+		private void saveContext()
+		{
+			mSavedDisplay = EGL14.eglGetCurrentDisplay();
+			mSavedContext = EGL14.eglGetCurrentContext();
+			mSavedSurfaceDraw = EGL14.eglGetCurrentSurface(EGL14.EGL_DRAW);
+			mSavedSurfaceRead = EGL14.eglGetCurrentSurface(EGL14.EGL_READ);
+		}
+
+		private void makeCurrent()
+		{
+			EGL14.eglMakeCurrent(mEglDisplay, mEglSurface, mEglSurface, mEglContext);
+		}
+
+		private void restoreContext()
+		{
+			// When working in a Vulkan environment we will not have a valid display, so just make sure we get rid of the current context
+			if (mSavedDisplay == EGL14.EGL_NO_DISPLAY)
+			{
+				EGL14.eglMakeCurrent(mEglDisplay, EGL14.EGL_NO_SURFACE, EGL14.EGL_NO_SURFACE, EGL14.EGL_NO_CONTEXT);
+			}
+			else
+			{
+				EGL14.eglMakeCurrent(mSavedDisplay, mSavedSurfaceDraw, mSavedSurfaceRead, mSavedContext);
+			}
+		}
+
+		private void initSurfaceTexture()
+		{
+			int[] textures = new int[1];
+			GLES20.glGenTextures(1, textures, 0);
+			mTextureID = textures[0];
+			if (mTextureID <= 0)
+			{
+				Log.e(TAG,"mTextureID <= 0");
+				release();
+				return;
+			}
+			mSurfaceTexture = new android.graphics.SurfaceTexture(mTextureID);
+			mSurface = new android.view.Surface(mSurfaceTexture);
+
+			int[] glInt = new int[1];
+
+			GLES20.glGenFramebuffers(1,glInt,0);
+			mFBO = glInt[0];
+			if (mFBO <= 0)
+			{
+				Log.e(TAG, "mFBO <= 0");
+				release();
+				return;
+			}
+			// Special shaders for blit of movie texture.
+			mBlitVertexShaderID = createShader(GLES20.GL_VERTEX_SHADER, mBlitVextexShader);
+			if (mBlitVertexShaderID == 0)
+			{
+				Log.e(TAG,"mBlitVertexShaderID == 0");
+				release();
+				return;
+			}
+			int mBlitFragmentShaderID = createShader(GLES20.GL_FRAGMENT_SHADER,
+				mVulkanRenderer ? (mSwizzlePixels ? mBlitFragmentShaderBGRA_NoConv : mBlitFragmentShaderRGBA_NoConv)
+								: (mSwizzlePixels ? mBlitFragmentShaderBGRA : mBlitFragmentShaderRGBA));
+			if (mBlitFragmentShaderID == 0)
+			{
+				Log.e(TAG,"mBlitFragmentShaderID == 0");
+				release();
+				return;
+			}
+			mProgram = GLES20.glCreateProgram();
+			if (mProgram <= 0)
+			{
+				Log.e(TAG,"mProgram <= 0");
+				release();
+				return;
+			}
+			GLES20.glAttachShader(mProgram, mBlitVertexShaderID);
+			GLES20.glAttachShader(mProgram, mBlitFragmentShaderID);
+			GLES20.glLinkProgram(mProgram);
+			int[] linkStatus = new int[1];
+			GLES20.glGetProgramiv(mProgram, GLES20.GL_LINK_STATUS, linkStatus, 0);
+			if (linkStatus[0] != GLES20.GL_TRUE)
+			{
+				Log.e(TAG,"Could not link program: ");
+				Log.e(TAG, GLES20.glGetProgramInfoLog(mProgram));
+				GLES20.glDeleteProgram(mProgram);
+				mProgram = 0;
+				release();
+				return;
+			}
+			mPositionAttrib = GLES20.glGetAttribLocation(mProgram, "Position");
+			mTexCoordsAttrib = GLES20.glGetAttribLocation(mProgram, "TexCoords");
+			mTextureUniform = GLES20.glGetUniformLocation(mProgram, "VideoTexture");
+
+			GLES20.glGenBuffers(1,glInt,0);
+			mBlitBuffer = glInt[0];
+			if (mBlitBuffer <= 0)
+			{
+				Log.e(TAG,"mBlitBuffer <= 0");
+				release();
+				return;
+			}
+
+			// Create blit mesh.
+			mTriangleVertices = java.nio.ByteBuffer.allocateDirect(
+				mTriangleVerticesData.length * FLOAT_SIZE_BYTES)
+					.order(java.nio.ByteOrder.nativeOrder()).asFloatBuffer();
+			mTriangleVerticesDirty = true;
+
+			// Set up GL state
+			if (mUseOwnContext)
+			{
+				GLES20.glDisable(GLES20.GL_BLEND);
+				GLES20.glDisable(GLES20.GL_CULL_FACE);
+				GLES20.glDisable(GLES20.GL_SCISSOR_TEST);
+				GLES20.glDisable(GLES20.GL_STENCIL_TEST);
+				GLES20.glDisable(GLES20.GL_DEPTH_TEST);
+				GLES20.glDisable(GLES20.GL_DITHER);
+				GLES20.glColorMask(true,true,true,true);
+			}
+		}
+
+		private void UpdateVertexData()
+		{
+			if (!mTriangleVerticesDirty || mBlitBuffer <= 0)
+			{
+				return;
+			}
+
+			// fill it in
+			mTriangleVertices.position(0);
+			mTriangleVertices.put(mTriangleVerticesData).position(0);
+
+			// save VBO state
+			int[] glInt = new int[1];
+			GLES20.glGetIntegerv(GLES20.GL_ARRAY_BUFFER_BINDING, glInt, 0);
+			int previousVBO = glInt[0];
+
+			GLES20.glBindBuffer(GLES20.GL_ARRAY_BUFFER, mBlitBuffer);
+			GLES20.glBufferData(GLES20.GL_ARRAY_BUFFER,
+				mTriangleVerticesData.length*FLOAT_SIZE_BYTES,
+				mTriangleVertices, GLES20.GL_STATIC_DRAW);
+
+			// restore VBO state
+			GLES20.glBindBuffer(GLES20.GL_ARRAY_BUFFER, previousVBO);
+
+			mTriangleVerticesDirty = false;
+		}
+
+		public boolean isValid()
+		{
+			return mSurfaceTexture != null;
+		}
+
+		private int createShader(int shaderType, String source)
+		{
+			int shader = GLES20.glCreateShader(shaderType);
+			if (shader != 0)
+			{
+				GLES20.glShaderSource(shader, source);
+				GLES20.glCompileShader(shader);
+				int[] compiled = new int[1];
+				GLES20.glGetShaderiv(shader, GLES20.GL_COMPILE_STATUS, compiled, 0);
+				if (compiled[0] == 0)
+				{
+					Log.e(TAG,"Could not compile shader " + shaderType + ":");
+					Log.e(TAG, GLES20.glGetShaderInfoLog(shader));
+					GLES20.glDeleteShader(shader);
+					shader = 0;
+				}
+			}
+			return shader;
+		}
+
+		public android.graphics.SurfaceTexture getSurfaceTexture()
+		{
+			return mSurfaceTexture;
+		}
+
+		public android.view.Surface getSurface()
+		{
+			return mSurface;
+		}
+
+		public int getExternalTextureId()
+		{
+			return mTextureID;
+		}
+
+		private static final int FLOAT_SIZE_BYTES = 4;
+		private static final int TRIANGLE_VERTICES_DATA_STRIDE_BYTES = 4 * FLOAT_SIZE_BYTES;
+		private static final int TRIANGLE_VERTICES_DATA_POS_OFFSET = 0;
+		private static final int TRIANGLE_VERTICES_DATA_UV_OFFSET = 2;
+		private float[] mTriangleVerticesData = {
+			// X, Y, U, V
+			-1.0f, -1.0f, 0.f, 0.f,
+			1.0f, -1.0f, 1.f, 0.f,
+			-1.0f, 1.0f, 0.f, 1.f,
+			1.0f, 1.0f, 1.f, 1.f,
+			};
+
+		private java.nio.FloatBuffer mTriangleVertices;
+
+		private final String mBlitVextexShader =
+			"attribute vec2 Position;\n" +
+			"attribute vec2 TexCoords;\n" +
+			"varying vec2 TexCoord;\n" +
+			"void main() {\n" +
+			"	TexCoord = TexCoords;\n" +
+			"	gl_Position = vec4(Position, 0.0, 1.0);\n" +
+			"}\n";
+
+		private final String mBlitFragmentShaderBGRA =
+			"#extension GL_OES_EGL_image_external : require\n" +
+			"uniform samplerExternalOES VideoTexture;\n" +
+			"varying highp vec2 TexCoord;\n" +
+			"void main()\n" +
+			"{\n" +
+			"	highp vec4 Color = texture2D(VideoTexture, TexCoord).bgra;\n" +
+			"	Color.rgb = max(Color.rgb, 6.10352e-5);\n" +
+			"	Color.rgb = mix(pow(Color.rgb * vec3(1.0 / 1.055) + vec3(0.0521327), vec3(2.4) ), Color.rgb * vec3(1.0 / 12.92), vec3(lessThanEqual(Color.rgb, vec3(0.04045))));\n" +
+			"	gl_FragColor = Color;\n" +
+			"}\n";
+		private final String mBlitFragmentShaderRGBA =
+			"#extension GL_OES_EGL_image_external : require\n" +
+			"uniform samplerExternalOES VideoTexture;\n" +
+			"varying highp vec2 TexCoord;\n" +
+			"void main()\n" +
+			"{\n" +
+			"	highp vec4 Color = texture2D(VideoTexture, TexCoord).rgba;\n" +
+			"	Color.rgb = max(Color.rgb, 6.10352e-5);\n" +
+			"	Color.rgb = mix(pow(Color.rgb * vec3(1.0 / 1.055) + vec3(0.0521327), vec3(2.4) ), Color.rgb * vec3(1.0 / 12.92), vec3(lessThanEqual(Color.rgb, vec3(0.04045))));\n" +
+			"	gl_FragColor = Color;\n" +
+			"}\n";
+		private final String mBlitFragmentShaderBGRA_NoConv =
+			"#extension GL_OES_EGL_image_external : require\n" +
+			"uniform samplerExternalOES VideoTexture;\n" +
+			"varying highp vec2 TexCoord;\n" +
+			"void main()\n" +
+			"{\n" +
+			"	gl_FragColor = texture2D(VideoTexture, TexCoord).bgra;;\n" +
+			"}\n";
+		private final String mBlitFragmentShaderRGBA_NoConv =
+			"#extension GL_OES_EGL_image_external : require\n" +
+			"uniform samplerExternalOES VideoTexture;\n" +
+			"varying highp vec2 TexCoord;\n" +
+			"void main()\n" +
+			"{\n" +
+			"	gl_FragColor = texture2D(VideoTexture, TexCoord).bgra;;\n" +
+			"}\n";
+
+		private int mProgram;
+		private int mPositionAttrib;
+		private int mTexCoordsAttrib;
+		private int mBlitBuffer;
+		private int mTextureUniform;
+
+		//
+		// Vulkan version: resulting image is returned as byte array in frame info data
+		//
+		public FFrameUpdateInfo updateFrameData()
+		{
+			synchronized(this)
+			{
+				FFrameUpdateInfo frameUpdateInfo = new FFrameUpdateInfo();
+				frameUpdateInfo.bFrameReady = false;
+				frameUpdateInfo.bRegionChanged = false;
+
+				int frameDataSize = mTextureWidth*mTextureHeight*4;
+				if ((null == mFrameData || frameDataSize != mFrameDataSize) && mTextureWidth > 0 && mTextureHeight > 0)
+				{
+					mFrameData = java.nio.ByteBuffer.allocateDirect(frameDataSize);
+					mFrameDataSize = frameDataSize;
+				}
+				// Copy surface texture to frame data.
+				if (!copyFrameTexture(0, mFrameData))
+				{
+					return null;
+				}
+
+				frameUpdateInfo.Buffer = mFrameData;
+
+				long ts = -1;
+				SurfaceTexture videoSurfaceTexture = getSurfaceTexture();
+				if (videoSurfaceTexture != null)
+				{
+					ts = videoSurfaceTexture.getTimestamp() / 1000; // nanos to micros
+				}
+				frameUpdateInfo.Timestamp = ts;
+				frameUpdateInfo.bFrameReady = true;
+				frameUpdateInfo.bRegionChanged = false;
+
+				return frameUpdateInfo;
+			}
+
+		}
+
+		//
+		// ES version: resulting image is written to given texture
+		//
+		public FFrameUpdateInfo updateFrameData(int destTexture)
+		{
+			synchronized(this)
+			{
+				// Copy surface texture to destination texture.
+				if (!copyFrameTexture(destTexture, null))
+				{
+					return null;
+				}
+
+				// Get associated frame metadata
+				FFrameUpdateInfo frameUpdateInfo = new FFrameUpdateInfo();
+
+				frameUpdateInfo.Buffer = null;
+
+				long ts = -1;
+				SurfaceTexture videoSurfaceTexture = getSurfaceTexture();
+				if (videoSurfaceTexture != null)
+				{
+					ts = videoSurfaceTexture.getTimestamp() / 1000; // nanos to micros
+				}
+				frameUpdateInfo.Timestamp = ts;
+
+				frameUpdateInfo.bFrameReady = true;
+				frameUpdateInfo.bRegionChanged = false;
+
+				return frameUpdateInfo;
+			}
+		}
+
+		// Copy the surface texture to another texture, or to raw data.
+		// Note: copying to raw data creates a temporary FBO texture.
+		private boolean copyFrameTexture(int destTexture, java.nio.Buffer destData)
+		{
+			if (null == mSurfaceTexture)
+			{
+				// Can't update if there's no surface to update into.
+				return false;
+			}
+
+			int[] glInt = new int[1];
+			boolean[] glBool = new boolean[1];
+			// Either use own context or save states
+			boolean previousBlend=false, previousCullFace=false, previousScissorTest=false, previousStencilTest=false, previousDepthTest=false, previousDither=false;
+			int previousFBO=0, previousVBO=0, previousMinFilter=0, previousMagFilter=0;
+			int[] previousViewport = new int[4];
+			if (mUseOwnContext)
+			{
+				// Received reports of these not being preserved when changing contexts
+				GLES20.glActiveTexture(GLES20.GL_TEXTURE0);
+				GLES20.glGetTexParameteriv(GLES20.GL_TEXTURE_2D, GLES20.GL_TEXTURE_MIN_FILTER, glInt, 0);
+				previousMinFilter = glInt[0];
+				GLES20.glGetTexParameteriv(GLES20.GL_TEXTURE_2D, GLES20.GL_TEXTURE_MAG_FILTER, glInt, 0);
+				previousMagFilter = glInt[0];
+
+				saveContext();
+				makeCurrent();
+			}
+			else
+			{
+				// Clear gl errors as they can creep in from the UE4 renderer.
+				GLES20.glGetError();
+
+				previousBlend = GLES20.glIsEnabled(GLES20.GL_BLEND);
+				previousCullFace = GLES20.glIsEnabled(GLES20.GL_CULL_FACE);
+				previousScissorTest = GLES20.glIsEnabled(GLES20.GL_SCISSOR_TEST);
+				previousStencilTest = GLES20.glIsEnabled(GLES20.GL_STENCIL_TEST);
+				previousDepthTest = GLES20.glIsEnabled(GLES20.GL_DEPTH_TEST);
+				previousDither = GLES20.glIsEnabled(GLES20.GL_DITHER);
+				GLES20.glGetIntegerv(GLES20.GL_FRAMEBUFFER_BINDING, glInt, 0);
+				previousFBO = glInt[0];
+				GLES20.glGetIntegerv(GLES20.GL_ARRAY_BUFFER_BINDING, glInt, 0);
+				previousVBO = glInt[0];
+				GLES20.glGetIntegerv(GLES20.GL_VIEWPORT, previousViewport, 0);
+
+				GLES20.glActiveTexture(GLES20.GL_TEXTURE0);
+				GLES20.glGetTexParameteriv(GLES20.GL_TEXTURE_2D, GLES20.GL_TEXTURE_MIN_FILTER, glInt, 0);
+				previousMinFilter = glInt[0];
+				GLES20.glGetTexParameteriv(GLES20.GL_TEXTURE_2D, GLES20.GL_TEXTURE_MAG_FILTER, glInt, 0);
+				previousMagFilter = glInt[0];
+
+				glVerify("save state");
+			}
+
+			// Get the latest video texture frame (the surface is really a FIFO of frames received from the decoder when the decoder output is released)
+			mSurfaceTexture.updateTexImage();
+			mSurfaceTexture.getTransformMatrix(mTransformMatrix);
+
+			float UMin = mTransformMatrix[12];
+			float UMax = UMin + mTransformMatrix[0];
+			float VMin = mTransformMatrix[13];
+			float VMax = VMin + mTransformMatrix[5];
+
+			if (mTriangleVerticesData[2] != UMin ||
+				mTriangleVerticesData[6] != UMax ||
+				mTriangleVerticesData[11] != VMin ||
+				mTriangleVerticesData[3] != VMax)
+			{
+				mTriangleVerticesData[ 2] = mTriangleVerticesData[10] = UMin;
+				mTriangleVerticesData[ 6] = mTriangleVerticesData[14] = UMax;
+				mTriangleVerticesData[11] = mTriangleVerticesData[15] = VMin;
+				mTriangleVerticesData[ 3] = mTriangleVerticesData[ 7] = VMax;
+				mTriangleVerticesDirty = true;
+			}
+
+			if (null != destData)
+			{
+				// Rewind data so that we can write to it.
+				destData.position(0);
+			}
+
+			if (!mUseOwnContext)
+			{
+				GLES20.glDisable(GLES20.GL_BLEND);
+				GLES20.glDisable(GLES20.GL_CULL_FACE);
+				GLES20.glDisable(GLES20.GL_SCISSOR_TEST);
+				GLES20.glDisable(GLES20.GL_STENCIL_TEST);
+				GLES20.glDisable(GLES20.GL_DEPTH_TEST);
+				GLES20.glDisable(GLES20.GL_DITHER);
+				GLES20.glColorMask(true,true,true,true);
+
+				glVerify("reset state");
+			}
+
+			GLES20.glViewport(0, 0, mTextureWidth, mTextureHeight);
+
+			glVerify("set viewport");
+
+			// Set-up FBO target texture..
+			int FBOTextureID = 0;
+			if (null != destData)
+			{
+				// Create temporary FBO for data copy.
+				GLES20.glGenTextures(1,glInt,0);
+				FBOTextureID = glInt[0];
+			}
+			else
+			{
+				// Use the given texture as the FBO.
+				FBOTextureID = destTexture;
+			}
+			// Set the FBO to draw into the texture one-to-one.
+			GLES20.glBindTexture(GLES20.GL_TEXTURE_2D, FBOTextureID);
+			GLES20.glTexParameteri(GLES20.GL_TEXTURE_2D,
+				GLES20.GL_TEXTURE_MIN_FILTER, GLES20.GL_NEAREST);
+			GLES20.glTexParameteri(GLES20.GL_TEXTURE_2D,
+				GLES20.GL_TEXTURE_MAG_FILTER, GLES20.GL_NEAREST);
+			GLES20.glTexParameteri(GLES20.GL_TEXTURE_2D,
+				GLES20.GL_TEXTURE_WRAP_S, GLES20.GL_CLAMP_TO_EDGE);
+			GLES20.glTexParameteri(GLES20.GL_TEXTURE_2D,
+				GLES20.GL_TEXTURE_WRAP_T, GLES20.GL_CLAMP_TO_EDGE);
+			// Create the temp FBO data if needed.
+			if (null != destData)
+			{
+				GLES20.glTexImage2D(GLES20.GL_TEXTURE_2D, 0,
+					GLES20.GL_RGBA,
+					mTextureWidth, mTextureHeight,
+					0, GLES20.GL_RGBA, GLES20.GL_UNSIGNED_BYTE, null);
+			}
+
+			glVerify("set-up FBO texture");
+			// Set to render to the FBO.
+			GLES20.glBindFramebuffer(GLES20.GL_FRAMEBUFFER, mFBO);
+			glVerify("glBindFramebuffer");
+			GLES20.glFramebufferTexture2D(
+				GLES20.GL_FRAMEBUFFER,
+				GLES20.GL_COLOR_ATTACHMENT0,
+				GLES20.GL_TEXTURE_2D, FBOTextureID, 0);
+
+			// check status
+			int status = GLES20.glCheckFramebufferStatus(GLES20.GL_FRAMEBUFFER);
+			if (status != GLES20.GL_FRAMEBUFFER_COMPLETE)
+			{
+				Log.w(TAG,"Failed to complete framebuffer attachment ("+status+")");
+			}
+
+			// The special shaders to render from the video texture.
+			GLES20.glUseProgram(mProgram);
+
+			// Set the mesh that renders the video texture.
+			UpdateVertexData();
+			GLES20.glBindBuffer(GLES20.GL_ARRAY_BUFFER, mBlitBuffer);
+			GLES20.glEnableVertexAttribArray(mPositionAttrib);
+			GLES20.glVertexAttribPointer(mPositionAttrib, 2, GLES20.GL_FLOAT, false,
+				TRIANGLE_VERTICES_DATA_STRIDE_BYTES, 0);
+			GLES20.glEnableVertexAttribArray(mTexCoordsAttrib);
+			GLES20.glVertexAttribPointer(mTexCoordsAttrib, 2, GLES20.GL_FLOAT, false,
+				TRIANGLE_VERTICES_DATA_STRIDE_BYTES,
+				TRIANGLE_VERTICES_DATA_UV_OFFSET*FLOAT_SIZE_BYTES);
+
+			glVerify("setup movie texture read");
+
+//GLES20.glClearColor(1.0f, 1.0f, 1.0f, 1.0f);
+//GLES20.glClear( GLES20.GL_COLOR_BUFFER_BIT);
+
+			// connect 'VideoTexture' to video source texture (mTextureID).
+			// mTextureID is bound to GL_TEXTURE_EXTERNAL_OES in updateTexImage
+			GLES20.glUniform1i(mTextureUniform, 0);
+			GLES20.glActiveTexture(GLES20.GL_TEXTURE0);
+			GLES20.glBindTexture(GL_TEXTURE_EXTERNAL_OES, mTextureID);
+
+			// Draw the video texture mesh.
+			GLES20.glDrawArrays(GLES20.GL_TRIANGLE_STRIP, 0, 4);
+
+			// Make sure GL issues commands to GPU
+			GLES20.glFlush();
+
+			// Read the FBO texture pixels into raw data.
+			if (null != destData)
+			{
+				GLES20.glReadPixels(
+					0, 0, mTextureWidth, mTextureHeight,
+					GLES20.GL_RGBA, GLES20.GL_UNSIGNED_BYTE,
+					destData);
+			}
+
+			glVerify("draw & read movie texture");
+
+			// Restore state and cleanup.
+			if (mUseOwnContext)
+			{
+				GLES20.glFramebufferTexture2D(
+					GLES20.GL_FRAMEBUFFER,
+					GLES20.GL_COLOR_ATTACHMENT0,
+					GLES20.GL_TEXTURE_2D, 0, 0);
+
+				if (null != destData && FBOTextureID > 0)
+				{
+					glInt[0] = FBOTextureID;
+					GLES20.glDeleteTextures(1, glInt, 0);
+				}
+
+				restoreContext();
+
+				// Restore previous texture filtering
+				GLES20.glTexParameteri(GLES20.GL_TEXTURE_2D, GLES20.GL_TEXTURE_MIN_FILTER, previousMinFilter);
+				GLES20.glTexParameteri(GLES20.GL_TEXTURE_2D, GLES20.GL_TEXTURE_MAG_FILTER, previousMagFilter);
+			}
+			else
+			{
+				GLES20.glBindFramebuffer(GLES20.GL_FRAMEBUFFER, previousFBO);
+				if (null != destData && FBOTextureID > 0)
+				{
+					glInt[0] = FBOTextureID;
+					GLES20.glDeleteTextures(1, glInt, 0);
+				}
+				GLES20.glBindBuffer(GLES20.GL_ARRAY_BUFFER, previousVBO);
+
+				GLES20.glViewport(previousViewport[0], previousViewport[1],	previousViewport[2], previousViewport[3]);
+				if (previousBlend) GLES20.glEnable(GLES20.GL_BLEND);
+				if (previousCullFace) GLES20.glEnable(GLES20.GL_CULL_FACE);
+				if (previousScissorTest) GLES20.glEnable(GLES20.GL_SCISSOR_TEST);
+				if (previousStencilTest) GLES20.glEnable(GLES20.GL_STENCIL_TEST);
+				if (previousDepthTest) GLES20.glEnable(GLES20.GL_DEPTH_TEST);
+				if (previousDither) GLES20.glEnable(GLES20.GL_DITHER);
+
+				GLES20.glTexParameteri(GLES20.GL_TEXTURE_2D, GLES20.GL_TEXTURE_MIN_FILTER, previousMinFilter);
+				GLES20.glTexParameteri(GLES20.GL_TEXTURE_2D, GLES20.GL_TEXTURE_MAG_FILTER, previousMagFilter);
+
+				// invalidate cached state in RHI
+				GLES20.glDisableVertexAttribArray(mPositionAttrib);
+				GLES20.glDisableVertexAttribArray(mTexCoordsAttrib);
+				nativeClearCachedAttributeState(mPositionAttrib, mTexCoordsAttrib); // offered by RHI for ES on Android
+			}
+
+			return true;
+		}
+
+		private void showGlError(String op, int error)
+		{
+			switch (error)
+			{
+				case GLES20.GL_INVALID_ENUM:						Log.e(TAG,"MediaPlayer$BitmapRenderer: " + op + ": glGetError GL_INVALID_ENUM");  break;
+				case GLES20.GL_INVALID_OPERATION:					Log.e(TAG,"MediaPlayer$BitmapRenderer: " + op + ": glGetError GL_INVALID_OPERATION");  break;
+				case GLES20.GL_INVALID_FRAMEBUFFER_OPERATION:		Log.e(TAG,"MediaPlayer$BitmapRenderer: " + op + ": glGetError GL_INVALID_FRAMEBUFFER_OPERATION");  break;
+				case GLES20.GL_INVALID_VALUE:						Log.e(TAG,"MediaPlayer$BitmapRenderer: " + op + ": glGetError GL_INVALID_VALUE");  break;
+				case GLES20.GL_FRAMEBUFFER_INCOMPLETE_ATTACHMENT:	Log.e(TAG,"MediaPlayer$BitmapRenderer: " + op + ": glGetError GL_FRAMEBUFFER_INCOMPLETE_ATTACHMENT");  break;
+				case GLES20.GL_FRAMEBUFFER_INCOMPLETE_DIMENSIONS:	Log.e(TAG,"MediaPlayer$BitmapRenderer: " + op + ": glGetError GL_FRAMEBUFFER_INCOMPLETE_DIMENSIONS");  break;
+				case GLES20.GL_FRAMEBUFFER_UNSUPPORTED:				Log.e(TAG,"MediaPlayer$BitmapRenderer: " + op + ": glGetError GL_FRAMEBUFFER_UNSUPPORTED");  break;
+				case GLES20.GL_OUT_OF_MEMORY:						Log.e(TAG,"MediaPlayer$BitmapRenderer: " + op + ": glGetError GL_OUT_OF_MEMORY");  break;
+				default:											Log.e(TAG,"MediaPlayer$BitmapRenderer: " + op + ": glGetError " + error);
+			}
+		}
+
+		private void glVerify(String op)
+		{
+			int error;
+			while ((error = GLES20.glGetError()) != GLES20.GL_NO_ERROR)
+			{
+				showGlError(op, error);
+				throw new RuntimeException(op + ": glGetError " + error);
+			}
+		}
+
+		private void glWarn(String op)
+		{
+			int error;
+			while ((error = GLES20.glGetError()) != GLES20.GL_NO_ERROR)
+			{
+				showGlError(op, error);
+			}
+		}
+
+		public FFrameUpdateInfo GetVideoFrameUpdateInfo(int destTexture, int width, int height)
+		{
+			mTextureWidth = width;
+			mTextureHeight = height;
+
+			if (mVulkanRenderer)
+			{
+				// Return data as byte buffer
+				return updateFrameData();
+			}
+			// Return data in ES texture
+			return updateFrameData(destTexture);
+		}
+
+		public void release()
+		{
+			synchronized(this)
+			{
+				if (mSurfaceTexture != null)
+				{
+					mSurfaceTexture.setOnFrameAvailableListener(null);
+				}
+			}
+
+			if (mEglContext != EGL14.EGL_NO_CONTEXT)
+			{
+				saveContext();
+				makeCurrent();
+			}
+			if (null != mSurface)
+			{
+				mSurface.release();
+				mSurface = null;
+			}
+			if (null != mSurfaceTexture)
+			{
+				mSurfaceTexture.release();
+				mSurfaceTexture = null;
+			}
+			int[] glInt = new int[1];
+			if (mBlitBuffer > 0)
+			{
+				glInt[0] = mBlitBuffer;
+				GLES20.glDeleteBuffers(1,glInt,0);
+				mBlitBuffer = -1;
+			}
+			if (mProgram > 0)
+			{
+				GLES20.glDeleteProgram(mProgram);
+				mProgram = -1;
+			}
+			if (mBlitVertexShaderID > 0)
+			{
+				GLES20.glDeleteShader(mBlitVertexShaderID);
+				mBlitVertexShaderID = -1;
+			}
+			if (mBlitFragmentShaderID > 0)
+			{
+				GLES20.glDeleteShader(mBlitFragmentShaderID);
+				mBlitFragmentShaderID = -1;
+			}
+			if (mFBO > 0)
+			{
+				glInt[0] = mFBO;
+				GLES20.glDeleteFramebuffers(1,glInt,0);
+				mFBO = -1;
+			}
+			if (mTextureID > 0)
+			{
+				glInt[0] = mTextureID;
+				GLES20.glDeleteTextures(1,glInt,0);
+				mTextureID = -1;
+			}
+			if (mEglSurface != EGL14.EGL_NO_SURFACE)
+			{
+				EGL14.eglDestroySurface(mEglDisplay, mEglSurface);
+				mEglSurface = EGL14.EGL_NO_SURFACE;
+			}
+			if (mEglContext != EGL14.EGL_NO_CONTEXT)
+			{
+				restoreContext();
+				EGL14.eglDestroyContext(mEglDisplay, mEglContext);
+				mEglContext = EGL14.EGL_NO_CONTEXT;
+			}
+			if (mCreatedEGLDisplay)
+			{
+				EGL14.eglTerminate(mEglDisplay);
+				mEglDisplay = EGL14.EGL_NO_DISPLAY;
+				mCreatedEGLDisplay = false;
+			}
+		}
+	}
+
+	public native void nativeClearCachedAttributeState(int PositionAttrib, int TexCoordsAttrib);
+}
+
+
diff --git a/Engine/Plugins/Media/ElectraUtil/Source/ElectraSamples/Public/Android/ElectraTextureSample.h b/Engine/Plugins/Media/ElectraUtil/Source/ElectraSamples/Public/Android/ElectraTextureSample.h
new file mode 100644
--- /dev/null
+++ b/Engine/Plugins/Media/ElectraUtil/Source/ElectraSamples/Public/Android/ElectraTextureSample.h
@@ -0,0 +1,163 @@
+// Copyright Epic Games, Inc. All Rights Reserved.
+
+#pragma once
+
+#include "HAL/Platform.h"
+
+#include "CoreTypes.h"
+#include "CoreGlobals.h"
+#include "Templates/SharedPointer.h"
+#include "IMediaTextureSample.h"
+#include "Math/IntPoint.h"
+#include "Math/Range.h"
+#include "Misc/Timespan.h"
+#include "MediaObjectPool.h"
+#include "MediaSamples.h"
+#include "MediaVideoDecoderOutputAndroid.h"
+
+#include "RHI.h"
+#include "RHIUtilities.h"
+
+class FElectraTextureSampleSupport;
+
+/**
+ * Texture sample generated by AndroidMedia player.
+ */
+class FElectraTextureSample final
+	: public IMediaTextureSample
+	, public IMediaPoolable
+	, public TSharedFromThis<FElectraTextureSample, ESPMode::ThreadSafe>
+{
+public:
+
+	/** Default constructor. */
+	FElectraTextureSample(TSharedPtr<FElectraTextureSampleSupport, ESPMode::ThreadSafe> InSupport)
+	: Support(InSupport)
+	, Texture(nullptr)
+	, Buffer(nullptr)
+	, BufferSize(0)
+	{
+	}
+
+	/** Virtual destructor. */
+	virtual ~FElectraTextureSample()
+	{
+		if (Buffer)
+		{
+			FMemory::Free(Buffer);
+		}
+	}
+
+	void Initialize(FVideoDecoderOutput* InVideoDecoderOutput);
+
+	virtual void InitializePoolable() override;
+	virtual void ShutdownPoolable() override;
+
+	//~ IMediaTextureSample interface
+
+	virtual const void* GetBuffer() override
+	{
+		return Buffer;
+	}
+
+	virtual FIntPoint GetDim() const override
+	{
+		return VideoDecoderOutput->GetDim();
+	}
+
+	virtual FTimespan GetDuration() const override
+	{
+		return VideoDecoderOutput->GetDuration();
+	}
+
+	virtual double GetAspectRatio() const override
+	{
+		return VideoDecoderOutput->GetAspectRatio();
+	}
+
+	virtual EMediaOrientation GetOrientation() const override
+	{
+		return (EMediaOrientation)VideoDecoderOutput->GetOrientation();
+	}
+
+	virtual EMediaTextureSampleFormat GetFormat() const override
+	{
+		return EMediaTextureSampleFormat::CharBGRA;
+	}
+
+	virtual FIntPoint GetOutputDim() const override
+	{
+		return VideoDecoderOutput->GetOutputDim();
+	}
+
+	virtual uint32 GetStride() const override
+	{
+		return GetDim().X * sizeof(uint32);
+	}
+
+	virtual FRHITexture* GetTexture() const override
+	{
+		return Texture.GetReference();
+	}
+
+	virtual FMediaTimeStamp GetTime() const override
+	{
+		FDecoderTimeStamp Time = VideoDecoderOutput->GetTime();
+		return FMediaTimeStamp(Time.Time, Time.SequenceIndex);
+	}
+
+	virtual bool IsCacheable() const override
+	{
+		return true;
+	}
+
+	virtual bool IsOutputSrgb() const override
+	{
+		return true;
+	}
+
+private:
+	friend class FElectraTextureSampleSupport;
+
+	void InitializeTexture();
+	void SetupFromBuffer(const void* InBuffer, int32 InBufferSize);
+
+	TSharedPtr<FElectraTextureSampleSupport, ESPMode::ThreadSafe> Support;
+
+	TSharedPtr<FVideoDecoderOutputAndroid, ESPMode::ThreadSafe> VideoDecoderOutput;
+
+	void* Buffer;
+	int32 BufferSize;
+
+	/** Texture resource. */
+	TRefCountPtr<FRHITexture2D> Texture;
+};
+
+
+using FElectraTextureSamplePtr  = TSharedPtr<FElectraTextureSample, ESPMode::ThreadSafe>;
+using FElectraTextureSampleRef  = TSharedRef<FElectraTextureSample, ESPMode::ThreadSafe>;
+
+class ELECTRASAMPLES_API FElectraTextureSamplePool : public TMediaObjectPool<FElectraTextureSample, FElectraTextureSamplePool>
+{
+	using ParentClass = TMediaObjectPool<FElectraTextureSample, FElectraTextureSamplePool>;
+	using TextureSample = FElectraTextureSample;
+
+public:
+	FElectraTextureSamplePool();
+
+	TextureSample* Alloc() const
+	{
+		return new TextureSample(Support);
+	}
+
+	void PrepareForDecoderShutdown()
+	{
+	}
+
+	void* GetCodecSurface();
+
+private:
+	TSharedPtr<FElectraTextureSampleSupport, ESPMode::ThreadSafe> Support;
+};
+
+
diff --git a/Engine/Plugins/Media/ElectraUtil/Source/ElectraSamples/Public/Apple/ElectraTextureSample.h b/Engine/Plugins/Media/ElectraUtil/Source/ElectraSamples/Public/Apple/ElectraTextureSample.h
new file mode 100644
--- /dev/null
+++ b/Engine/Plugins/Media/ElectraUtil/Source/ElectraSamples/Public/Apple/ElectraTextureSample.h
@@ -0,0 +1,154 @@
+// Copyright Epic Games, Inc. All Rights Reserved.
+
+#pragma once
+
+#include "HAL/Platform.h"
+
+#include "MediaSamples.h"
+#include "IMediaTextureSampleConverter.h"
+#include "Misc/Timespan.h"
+#include "MediaObjectPool.h"
+#include "MediaVideoDecoderOutputApple.h"
+
+#import <AVFoundation/AVFoundation.h>
+#import <VideoToolbox/VideoToolbox.h>
+
+class FElectraMediaTexConvApple;
+
+
+class ELECTRASAMPLES_API FElectraTextureSample final
+	: public IMediaTextureSample
+    , public IMediaTextureSampleConverter
+	, public IMediaPoolable
+{
+public:
+	FElectraTextureSample(const TWeakPtr<FElectraMediaTexConvApple, ESPMode::ThreadSafe> & InTexConv)
+		: Texture(nullptr)
+		, TexConv(InTexConv)
+	{
+	}
+
+	virtual ~FElectraTextureSample()
+	{
+	}
+
+public:
+	void Initialize(FVideoDecoderOutput* InVideoDecoderOutput)
+	{
+		VideoDecoderOutput = StaticCastSharedPtr<FVideoDecoderOutputApple, IDecoderOutputPoolable, ESPMode::ThreadSafe>(InVideoDecoderOutput->AsShared());
+	}
+
+	//~ IMediaTextureSample interface
+
+	virtual const void* GetBuffer() override
+	{
+		return nullptr;
+	}
+
+	virtual FIntPoint GetDim() const override;
+
+	virtual FTimespan GetDuration() const override;
+
+	virtual double GetAspectRatio() const override
+	{
+		return VideoDecoderOutput->GetAspectRatio();
+	}
+
+	virtual EMediaOrientation GetOrientation() const override
+	{
+		return (EMediaOrientation)VideoDecoderOutput->GetOrientation();
+	}
+
+	virtual EMediaTextureSampleFormat GetFormat() const override
+	{
+		return EMediaTextureSampleFormat::CharBGRA;
+	}
+
+	virtual FIntPoint GetOutputDim() const override;
+	virtual uint32 GetStride() const override;
+
+	virtual FRHITexture* GetTexture() const override
+	{
+		return Texture;
+	}
+
+	virtual FMediaTimeStamp GetTime() const override;
+
+	virtual bool IsCacheable() const override
+	{
+		return true;
+	}
+
+	virtual bool IsOutputSrgb() const override
+	{
+		return true;
+	}
+
+    virtual IMediaTextureSampleConverter* GetMediaTextureSampleConverter() override
+    {
+        return this;
+    }
+
+	virtual void InitializePoolable() override;
+	virtual void ShutdownPoolable() override;
+
+private:
+	/** The sample's texture resource. */
+	TRefCountPtr<FRHITexture2D> Texture;
+
+	/** Output data from video decoder. */
+	TSharedPtr<FVideoDecoderOutputApple, ESPMode::ThreadSafe> VideoDecoderOutput;
+
+	TWeakPtr<FElectraMediaTexConvApple, ESPMode::ThreadSafe> TexConv;
+
+	virtual uint32 GetConverterInfoFlags() const override;
+    virtual bool Convert(FTexture2DRHIRef & InDstTexture, const FConversionHints & Hints) override;
+};
+
+
+using FElectraTextureSamplePtr  = TSharedPtr<FElectraTextureSample, ESPMode::ThreadSafe>;
+using FElectraTextureSampleRef  = TSharedRef<FElectraTextureSample, ESPMode::ThreadSafe>;
+
+
+class ELECTRASAMPLES_API FElectraMediaTexConvApple
+{
+public:
+    FElectraMediaTexConvApple();
+    ~FElectraMediaTexConvApple();
+
+#if WITH_ENGINE
+    void ConvertTexture(FTexture2DRHIRef & InDstTexture, CVImageBufferRef InImageBufferRef);
+#endif
+
+private:
+#if WITH_ENGINE
+    /** The Metal texture cache for unbuffered texture uploads. */
+    CVMetalTextureCacheRef MetalTextureCache;
+#endif
+};
+
+
+class ELECTRASAMPLES_API FElectraTextureSamplePool : public TMediaObjectPool<FElectraTextureSample, FElectraTextureSamplePool>
+{
+	using ParentClass = TMediaObjectPool<FElectraTextureSample, FElectraTextureSamplePool>;
+	using TextureSample = FElectraTextureSample;
+
+public:
+	FElectraTextureSamplePool()
+		: TMediaObjectPool<TextureSample, FElectraTextureSamplePool>(this)
+		, TexConv(new FElectraMediaTexConvApple())
+	{}
+
+	TextureSample *Alloc() const
+	{
+		return new TextureSample(TexConv);
+	}
+
+	void PrepareForDecoderShutdown()
+	{
+	}
+
+private:
+	TSharedPtr<FElectraMediaTexConvApple, ESPMode::ThreadSafe> TexConv;
+};
+
diff --git a/Engine/Plugins/Media/ElectraUtil/Source/ElectraSamples/Public/ElectraSamplesTextureManager.h b/Engine/Plugins/Media/ElectraUtil/Source/ElectraSamples/Public/ElectraSamplesTextureManager.h
new file mode 100644
--- /dev/null
+++ b/Engine/Plugins/Media/ElectraUtil/Source/ElectraSamples/Public/ElectraSamplesTextureManager.h
@@ -0,0 +1,31 @@
+// Copyright Epic Games, Inc. All Rights Reserved.
+
+#pragma once
+
+#include "RHI.h"
+#include "RHIResources.h"
+#include "MediaVideoDecoderOutput.h"
+
+class FElectraSamplesTextureManager
+{
+public:
+	FElectraSamplesTextureManager();
+	virtual ~FElectraSamplesTextureManager();
+
+	TSharedPtr<IVideoDecoderTexture, ESPMode::ThreadSafe> CreateTexture(const FIntPoint& Dim, EPixelFormat Fmt = PF_Unknown);
+	FTexture2DRHIRef GetRHITextureFromPlatformTexture(void* PlatformTexture);
+
+protected:
+	virtual TSharedPtr<IVideoDecoderTexture, ESPMode::ThreadSafe> PlatformCreateTexture(const FIntPoint& Dim, EPixelFormat Fmt) = 0;
+	virtual FTexture2DRHIRef PlatformGetTextureRHI(IVideoDecoderTexture *Texture) = 0;
+	virtual void* PlatformGetTexturePlatform(IVideoDecoderTexture* Texture) = 0;
+	
+	void PlatformAddSharedTextureRef(TSharedPtr<IVideoDecoderTexture, ESPMode::ThreadSafe> InTextureRef);
+	void PlatformRemoveSharedTextureRef(TSharedPtr<IVideoDecoderTexture, ESPMode::ThreadSafe> InTextureRef);
+
+private:
+	void CleanupMap();
+
+	FCriticalSection AccessCS;
+	TMap<void*, TWeakPtr<IVideoDecoderTexture, ESPMode::ThreadSafe>> KnownTextures;
+};
diff --git a/Engine/Plugins/Media/ElectraUtil/Source/ElectraSamples/Public/GStreamer/ElectraTextureSample.h b/Engine/Plugins/Media/ElectraUtil/Source/ElectraSamples/Public/GStreamer/ElectraTextureSample.h
new file mode 100644
--- /dev/null
+++ b/Engine/Plugins/Media/ElectraUtil/Source/ElectraSamples/Public/GStreamer/ElectraTextureSample.h
@@ -0,0 +1,146 @@
+// Copyright Epic Games, Inc. All Rights Reserved.
+
+#pragma once
+
+#include "CoreMinimal.h"
+#include "IMediaTextureSample.h"
+#include "MediaObjectPool.h"
+#include "RHI.h"
+#include "RHIUtilities.h"
+
+class MEDIArendererVideoUE;
+class FRHITexture2D;
+
+class FElectraTextureSampleUnix 
+	: public IMediaTextureSample
+	, public IMediaPoolable
+{
+public:
+	FElectraTextureSampleUnix()
+		: Time(FTimespan::Zero())
+	{
+	}
+
+public:
+	void Initialize(FIntPoint InDisplaySize, FIntPoint InTotalSize, const FMediaTimeStamp & InTime, FTimespan InDuration)
+	{
+		Time = InTime;
+		DisplaySize = InDisplaySize;
+		TotalSize = InTotalSize;
+		Duration = InDuration;
+	}
+
+	void CreateTexture()
+	{
+		check(IsInRenderingThread());
+
+		const uint32 CreateFlags = TexCreate_Dynamic | TexCreate_SRGB;
+
+		TRefCountPtr<FRHITexture2D> DummyTexture2DRHI;
+		FRHIResourceCreateInfo CreateInfo;
+
+		RHICreateTargetableShaderResource2D(
+			TotalSize.X,
+			TotalSize.Y,
+			PF_B8G8R8A8,
+			1,
+			CreateFlags,
+			TexCreate_RenderTargetable,
+			false,
+			CreateInfo,
+			Texture,
+			DummyTexture2DRHI
+		);
+	}	
+
+public:
+	//~ IMediaTextureSample interface
+	virtual const void* GetBuffer() override
+	{
+		return nullptr;
+	}
+
+	virtual FIntPoint GetDim() const override
+	{
+		return TotalSize;
+	}
+
+	virtual FTimespan GetDuration() const override
+	{
+		return Duration;
+	}
+
+	virtual EMediaTextureSampleFormat GetFormat() const override
+	{
+		return EMediaTextureSampleFormat::CharBGRA;
+	}
+
+	virtual FIntPoint GetOutputDim() const override
+	{
+		return DisplaySize;
+	}
+
+	virtual uint32 GetStride() const override
+	{
+		if (!Texture.IsValid())
+		{
+			return 0;
+		}
+
+		return Texture->GetSizeX() * 4;
+	}
+
+
+	virtual FRHITexture* GetTexture() const override
+	{
+		return Texture.GetReference();
+	}
+
+	virtual FMediaTimeStamp GetTime() const override
+	{
+		return Time;
+	}
+
+	virtual bool IsCacheable() const override
+	{
+		// ??? 
+		return true;
+	}
+
+	virtual bool IsOutputSrgb() const override
+	{
+		return true;
+	}
+
+
+public:
+	//~ IMediaPoolable interface
+#if !UE_SERVER
+	virtual void InitializePoolable() override;
+	virtual void ShutdownPoolable() override;
+#endif
+
+public:
+	TRefCountPtr<FRHITexture2D> GetTextureRef() const
+	{
+		return Texture;
+	}
+
+	// We hold a weak reference to the video renderer. During destruction the video renderer could be destroyed while samples are still out there..
+	TWeakPtr<MEDIArendererVideoUE, ESPMode::ThreadSafe> OwningRenderer;
+
+private:
+	TRefCountPtr<FRHITexture2D> Texture;
+	FMediaTimeStamp Time;
+	FTimespan Duration;
+	FIntPoint TotalSize;
+	FIntPoint DisplaySize;
+};
+
+using FElectraTextureSamplePtr = TSharedPtr<FElectraTextureSampleUnix, ESPMode::ThreadSafe>;
+
+class FElectraTextureSamplePool : public TMediaObjectPool<FElectraTextureSampleUnix>
+{
+public:
+	void PrepareForDecoderShutdown() {}
+};
diff --git a/Engine/Plugins/Media/ElectraUtil/Source/ElectraSamples/Public/Windows/ElectraTextureSample.h b/Engine/Plugins/Media/ElectraUtil/Source/ElectraSamples/Public/Windows/ElectraTextureSample.h
new file mode 100644
--- /dev/null
+++ b/Engine/Plugins/Media/ElectraUtil/Source/ElectraSamples/Public/Windows/ElectraTextureSample.h
@@ -0,0 +1,114 @@
+// Copyright Epic Games, Inc. All Rights Reserved.
+#pragma once
+
+#include "IMediaTextureSample.h"
+#include "IMediaTextureSampleConverter.h"
+#include "MediaObjectPool.h"
+#include "Misc/Timespan.h"
+#include "Templates/SharedPointer.h"
+#include "Templates/RefCounting.h"
+#include "RHI.h"
+#include "RHIResources.h"
+#include "ShaderCore.h"
+
+#include "MediaVideoDecoderOutputPC.h"
+
+#if !PLATFORM_WINDOWS
+#error "Should only be used on Windows"
+#endif
+
+class ELECTRASAMPLES_API FElectraTextureSample final
+	: public IMediaTextureSample
+	, public IMediaPoolable
+	, public IMediaTextureSampleConverter
+{
+public:
+
+	/** Default constructor. */
+	FElectraTextureSample()
+		: SampleFormat(EMediaTextureSampleFormat::Undefined)
+	{ }
+
+	void Initialize(FVideoDecoderOutput* VideoDecoderOutput);
+	FVideoDecoderOutput *GetDecoderOutput() { return VideoDecoderOutput.Get(); }
+
+	//
+	// General Interface
+	//
+
+	virtual const void* GetBuffer() override;
+	virtual uint32 GetStride() const override;
+
+	virtual FIntPoint GetDim() const override;
+	virtual FIntPoint GetOutputDim() const override;
+
+	virtual FMediaTimeStamp GetTime() const override;
+	virtual FTimespan GetDuration() const override;
+
+	virtual double GetAspectRatio() const override
+	{
+		return VideoDecoderOutput->GetAspectRatio();
+	}
+
+	virtual EMediaOrientation GetOrientation() const override
+	{
+		return (EMediaOrientation)VideoDecoderOutput->GetOrientation();
+	}
+
+	virtual EMediaTextureSampleFormat GetFormat() const override
+	{
+		return SampleFormat;
+	}
+
+	virtual bool IsCacheable() const override
+	{
+		return true;
+	}
+
+	virtual bool IsOutputSrgb() const override
+	{
+		return true;
+	}
+
+#if WITH_ENGINE
+	virtual FRHITexture* GetTexture() const override
+	{
+		return Texture;
+	}
+#endif //WITH_ENGINE
+
+	IMFSample* GetMFSample();
+
+#if !UE_SERVER
+	virtual void InitializePoolable() override;
+	virtual void ShutdownPoolable() override;
+#endif
+
+	virtual IMediaTextureSampleConverter* GetMediaTextureSampleConverter() override;
+
+private:
+	virtual bool Convert(FTexture2DRHIRef & InDstTexture, const FConversionHints & Hints) override;
+	virtual uint32 GetConverterInfoFlags() const
+	{
+		return ConverterInfoFlags_PreprocessOnly;
+	}
+
+	/** The sample format. */
+	EMediaTextureSampleFormat SampleFormat;
+
+	/** Destination Texture resource (from Rendering device) */
+	FTexture2DRHIRef Texture;
+
+	/** Output data from video decoder. */
+	TSharedPtr<FVideoDecoderOutputPC, ESPMode::ThreadSafe> VideoDecoderOutput;
+};
+
+using FElectraTextureSamplePtr = TSharedPtr<FElectraTextureSample, ESPMode::ThreadSafe>;
+using FElectraTextureSampleRef = TSharedRef<FElectraTextureSample, ESPMode::ThreadSafe>;
+
+class FElectraTextureSamplePool : public TMediaObjectPool<FElectraTextureSample>
+{
+public:
+	void PrepareForDecoderShutdown() {}
+};
+
