Android Build fix

diff --git a/Engine/Plugins/Media/ElectraUtil/Source/ElectraSamples/Private/ElectraTextureSampleAndroid.cpp b/Engine/Plugins/Media/ElectraUtil/Source/ElectraSamples/Private/ElectraTextureSampleAndroid.cpp
--- a/Engine/Plugins/Media/ElectraUtil/Source/ElectraSamples/Private/ElectraTextureSampleAndroid.cpp
+++ b/Engine/Plugins/Media/ElectraUtil/Source/ElectraSamples/Private/ElectraTextureSampleAndroid.cpp
@@ -423,7 +423,7 @@ void FElectraTextureSample::InitializeTexture()
 	}
 
 	TRefCountPtr<FRHITexture2D> DummyTexture2DRHI;
-	FRHIResourceCreateInfo CreateInfo;
+	FRHIResourceCreateInfo CreateInfo(TEXT("FElectraTextureSample"));
 
 	RHICreateTargetableShaderResource2D(
 		Dim.X,
