ElectraPlayer: Moved PlayerTime.cpp out into ElectraBase where the .h file is already located so it can be used by other Electra plugins.

#rnx

diff --git a/Engine/Plugins/Media/ElectraPlayer/Source/ElectraPlayerRuntime/Private/Runtime/Platforms/Android/PlayerRuntime_Android.cpp b/Engine/Plugins/Media/ElectraPlayer/Source/ElectraPlayerRuntime/Private/Runtime/Platforms/Android/PlayerRuntime_Android.cpp
--- a/Engine/Plugins/Media/ElectraPlayer/Source/ElectraPlayerRuntime/Private/Runtime/Platforms/Android/PlayerRuntime_Android.cpp
+++ b/Engine/Plugins/Media/ElectraPlayer/Source/ElectraPlayerRuntime/Private/Runtime/Platforms/Android/PlayerRuntime_Android.cpp
@@ -1,6 +1,9 @@
 // Copyright Epic Games, Inc. All Rights Reserved.
 
-#include "PlayerTime.h"
+#include "SynchronizedClock.h"
+
+namespace Electra
+{
 
 int64 MEDIAutcTime::CurrentMSec()
 {
@@ -8,3 +11,4 @@ int64 MEDIAutcTime::CurrentMSec()
 	return (int64)localTime.GetTotalMilliseconds();
 };
 
+}
diff --git a/Engine/Plugins/Media/ElectraPlayer/Source/ElectraPlayerRuntime/Private/Runtime/Platforms/Apple/PlayerRuntime_Apple.cpp b/Engine/Plugins/Media/ElectraPlayer/Source/ElectraPlayerRuntime/Private/Runtime/Platforms/Apple/PlayerRuntime_Apple.cpp
--- a/Engine/Plugins/Media/ElectraPlayer/Source/ElectraPlayerRuntime/Private/Runtime/Platforms/Apple/PlayerRuntime_Apple.cpp
+++ b/Engine/Plugins/Media/ElectraPlayer/Source/ElectraPlayerRuntime/Private/Runtime/Platforms/Apple/PlayerRuntime_Apple.cpp
@@ -1,6 +1,9 @@
 // Copyright Epic Games, Inc. All Rights Reserved.
 
-#include "PlayerTime.h"
+#include "SynchronizedClock.h"
+
+namespace Electra
+{
 
 int64 MEDIAutcTime::CurrentMSec()
 {
@@ -8,3 +11,4 @@ int64 MEDIAutcTime::CurrentMSec()
 	return (int64)localTime.GetTotalMilliseconds();
 };
 
+}
diff --git a/Engine/Plugins/Media/ElectraPlayer/Source/ElectraPlayerRuntime/Private/Runtime/Platforms/Windows/PlayerRuntime_PC.cpp b/Engine/Plugins/Media/ElectraPlayer/Source/ElectraPlayerRuntime/Private/Runtime/Platforms/Windows/PlayerRuntime_PC.cpp
--- a/Engine/Plugins/Media/ElectraPlayer/Source/ElectraPlayerRuntime/Private/Runtime/Platforms/Windows/PlayerRuntime_PC.cpp
+++ b/Engine/Plugins/Media/ElectraPlayer/Source/ElectraPlayerRuntime/Private/Runtime/Platforms/Windows/PlayerRuntime_PC.cpp
@@ -1,6 +1,9 @@
 // Copyright Epic Games, Inc. All Rights Reserved.
 
-#include "PlayerTime.h"
+#include "SynchronizedClock.h"
+
+namespace Electra
+{
 
 int64 MEDIAutcTime::CurrentMSec()
 {
@@ -8,3 +11,4 @@ int64 MEDIAutcTime::CurrentMSec()
 	return (int64)localTime.GetTotalMilliseconds();
 };
 
+}
diff --git a/Engine/Plugins/Media/ElectraPlayer/Source/ElectraPlayerRuntime/Private/Runtime/PlayerTime.cpp b/Engine/Plugins/Media/ElectraPlayer/Source/ElectraPlayerRuntime/Private/Runtime/PlayerTime.cpp
deleted file mode 100644
diff --git a/Engine/Plugins/Media/ElectraPlayer/Source/ElectraPlayerRuntime/Private/Runtime/SynchronizedClock.h b/Engine/Plugins/Media/ElectraPlayer/Source/ElectraPlayerRuntime/Private/Runtime/SynchronizedClock.h
--- a/Engine/Plugins/Media/ElectraPlayer/Source/ElectraPlayerRuntime/Private/Runtime/SynchronizedClock.h
+++ b/Engine/Plugins/Media/ElectraPlayer/Source/ElectraPlayerRuntime/Private/Runtime/SynchronizedClock.h
@@ -36,6 +36,19 @@ namespace Electra
 	};
 
 
+	/**
+	 * System wallclock time (UTC)
+	**/
+	class MEDIAutcTime
+	{
+	public:
+		static Electra::FTimeValue Current()
+		{
+			return Electra::FTimeValue(Electra::FTimeValue::MillisecondsToHNS(CurrentMSec()));
+		}
+		static int64 CurrentMSec();
+	};
+
 } // namespace Electra
 
 
diff --git a/Engine/Plugins/Media/ElectraUtil/Source/ElectraBase/Private/PlayerTime.cpp b/Engine/Plugins/Media/ElectraUtil/Source/ElectraBase/Private/PlayerTime.cpp
new file mode 100644
--- /dev/null
+++ b/Engine/Plugins/Media/ElectraUtil/Source/ElectraBase/Private/PlayerTime.cpp
@@ -0,0 +1,245 @@
+// Copyright Epic Games, Inc. All Rights Reserved.
+
+#include "PlayerTime.h"
+
+#include "Math/BigInt.h"
+
+
+namespace Electra
+{
+
+	namespace TimeStringHelpers
+	{
+		int32 FindFirstNotOf(const FString& InString, const FString& InNotOfChars, int32 FirstPos = 0)
+		{
+			for (int32 i = FirstPos; i < InString.Len(); ++i)
+			{
+				bool bFoundCharFromNotOfChars = false;
+				for (int32 j = 0; j < InNotOfChars.Len(); ++j)
+				{
+					if (InString[i] == InNotOfChars[j])
+					{
+						// We found a character from the "NOT of" list. Check next...
+						bFoundCharFromNotOfChars = true;
+						break;
+					}
+				}
+				if (!bFoundCharFromNotOfChars)
+				{
+					// We did not find any of the characters. This is what we are looking for and return the index of the first "not of" character
+					return i;
+				}
+			}
+			return INDEX_NONE;
+		}
+
+		int32 FindLastNotOf(const FString& InString, const FString& InNotOfChars, int32 InStartPos = MAX_int32)
+		{
+			InStartPos = FMath::Min(InStartPos, InString.Len() - 1);
+			for (int32 i = InStartPos; i >= 0; --i)
+			{
+				bool bFoundCharFromNotOfChars = false;
+				for (int32 j = 0; j < InNotOfChars.Len(); ++j)
+				{
+					if (InString[i] == InNotOfChars[j])
+					{
+						// We found a character from the "NOT of" list. Check next...
+						bFoundCharFromNotOfChars = true;
+						break;
+					}
+				}
+				if (!bFoundCharFromNotOfChars)
+				{
+					// We did not find any of the characters. This is what we are looking for and return the index of the first "not of" character
+					return i;
+				}
+			}
+			return INDEX_NONE;
+		}
+	}
+
+
+	FTimeValue& FTimeValue::SetFromTimeFraction(const FTimeFraction& TimeFraction)
+	{
+		if (TimeFraction.IsValid())
+		{
+			if (TimeFraction.IsPositiveInfinity())
+			{
+				SetToPositiveInfinity();
+			}
+			else
+			{
+				HNS = TimeFraction.GetAsTimebase(10000000);
+				bIsInfinity = false;
+				bIsValid = true;
+			}
+		}
+		else
+		{
+			SetToInvalid();
+		}
+		return *this;
+	}
+
+	FTimeValue& FTimeValue::SetFromND(int64 Numerator, uint32 Denominator)
+	{
+		if (Denominator != 0)
+		{
+			if (Denominator == 10000000)
+			{
+				HNS 		= Numerator;
+				bIsValid	= true;
+				bIsInfinity = false;
+			}
+			else if (Numerator >= -922337203685LL && Numerator <= 922337203685LL)
+			{
+				HNS 		= Numerator * 10000000 / Denominator;
+				bIsValid	= true;
+				bIsInfinity = false;
+			}
+			else
+			{
+				SetFromTimeFraction(FTimeFraction(Numerator, Denominator));
+			}
+		}
+		else
+		{
+			HNS 		= Numerator>=0 ? 0x7fffffffffffffffLL : -0x7fffffffffffffffLL;
+			bIsValid	= true;
+			bIsInfinity = true;
+		}
+		return *this;
+	}
+
+	//-----------------------------------------------------------------------------
+	/**
+	 * Returns this time value in a custom timebase. Requires internal bigint conversion and is therefor SLOW!
+	 *
+	 * @param CustomTimebase
+	 *
+	 * @return
+	 */
+	int64 FTimeValue::GetAsTimebase(uint32 CustomTimebase) const
+	{
+		// Some shortcuts
+		if (!bIsValid)
+		{
+			return 0;
+		}
+		else if (bIsInfinity)
+		{
+			return HNS >= 0 ? 0x7fffffffffffffffLL : -0x7fffffffffffffffLL;
+		}
+		else if (HNS == 0)
+		{
+			return 0;
+		}
+
+		TBigInt<128> n(HNS);
+		TBigInt<128> d(10000000);
+		TBigInt<128> s(CustomTimebase);
+
+		n *= s;
+		n /= d;
+
+		int64 r = n.ToInt();
+		return r;
+	}
+
+	// ---------------------------------------------------------------------------------------------------------------------
+
+	FTimeFraction& FTimeFraction::SetFromFloatString(const FString& InString)
+	{
+		// The string value must consist only of a sign, decimal digits and an optional period.
+		static const FString kTextDigitsEtc(TEXT("0123456789.-+"));
+		static const FString kTextZero(TEXT("0"));
+		if (TimeStringHelpers::FindFirstNotOf(InString, kTextDigitsEtc) == INDEX_NONE) 
+		{
+			Denominator = 1;
+			int32 DotIndex;
+			InString.FindChar(TCHAR('.'), DotIndex);
+			if (DotIndex == INDEX_NONE)
+			{
+				LexFromString(Numerator, *InString);
+				bIsValid = true;
+			}
+			else
+			{
+				LexFromString(Numerator, *(InString.Mid(0, DotIndex)));
+				FString frc = InString.Mid(DotIndex + 1);
+				// Remove all trailing zeros
+				int32 last0 = TimeStringHelpers::FindLastNotOf(frc, kTextZero);
+				if (last0 != INDEX_NONE)
+				{
+					frc.MidInline(0, last0 + 1);
+					// Convert at most 7 fractional digits (giving us hundreds of nanoseconds (HNS))
+					for(int32 i = 0; i < frc.Len() && i<7; ++i)
+					{
+						Numerator = Numerator * 10 + (frc[i] - TCHAR('0'));
+						Denominator *= 10;
+					}
+				}
+				bIsValid = true;
+			}
+		}
+		else
+		{
+			static const FString kInf0(TEXT("INF"));
+			static const FString kInf1(TEXT("+INF"));
+			static const FString kInf2(TEXT("-INF"));
+			if (InString.Equals(kInf0) || InString.Equals(kInf1))
+			{
+				Numerator = 1;
+				Denominator = 0;
+				bIsValid = true;
+			}
+			else if (InString.Equals(kInf2))
+			{
+				Numerator = -1;
+				Denominator = 0;
+				bIsValid = true;
+			}
+			else
+			{
+				bIsValid = false;
+			}
+		}
+		return *this;
+	}
+
+
+	int64 FTimeFraction::GetAsTimebase(uint32 CustomTimebase) const
+	{
+		// Some shortcuts
+		if (!bIsValid)
+		{
+			return 0;
+		}
+		else if (CustomTimebase == Denominator)
+		{
+			return Numerator;
+		}
+		else if (Numerator == 0)
+		{
+			return 0;
+		}
+		// Infinity?
+		else if (Denominator == 0 || CustomTimebase == 0)
+		{
+			return Numerator >= 0 ? 0x7fffffffffffffffLL : -0x7fffffffffffffffLL;
+		}
+
+		TBigInt<128> n(Numerator);
+		TBigInt<128> d(Denominator);
+		TBigInt<128> s(CustomTimebase);
+
+		n *= s;
+		n /= d;
+
+		int64 r = n.ToInt();
+		return r;
+	}
+
+
+}
+
diff --git a/Engine/Plugins/Media/ElectraUtil/Source/ElectraBase/Public/PlayerTime.h b/Engine/Plugins/Media/ElectraUtil/Source/ElectraBase/Public/PlayerTime.h
--- a/Engine/Plugins/Media/ElectraUtil/Source/ElectraBase/Public/PlayerTime.h
+++ b/Engine/Plugins/Media/ElectraUtil/Source/ElectraBase/Public/PlayerTime.h
@@ -13,7 +13,7 @@ class FTimeFraction;
 /**
  * Keeps a time value in hundred nanoseconds (HNS).
  */
-class FTimeValue
+class ELECTRABASE_API FTimeValue
 {
 public:
 	static FTimeValue GetInvalid()
@@ -479,7 +479,7 @@ struct FTimeRange
 /**
  * Keeps a time value as a fractional.
  */
-class FTimeFraction
+class ELECTRABASE_API FTimeFraction
 {
 public:
 	static const FTimeFraction& GetInvalid()
@@ -595,19 +595,3 @@ private:
 
 }
 
-
-
-/**
- * System wallclock time (UTC)
-**/
-class MEDIAutcTime
-{
-public:
-	static Electra::FTimeValue Current()
-	{
-		return Electra::FTimeValue(Electra::FTimeValue::MillisecondsToHNS(CurrentMSec()));
-	}
-	static int64 CurrentMSec();
-};
-
-
