ElectraPlayer: Moved several utility functions to ElectraBase

diff --git a/Engine/Plugins/Media/ElectraPlayer/Source/ElectraPlayerRuntime/Private/Runtime/Player/DASH/ManifestParserDASH.cpp b/Engine/Plugins/Media/ElectraPlayer/Source/ElectraPlayerRuntime/Private/Runtime/Player/DASH/ManifestParserDASH.cpp
--- a/Engine/Plugins/Media/ElectraPlayer/Source/ElectraPlayerRuntime/Private/Runtime/Player/DASH/ManifestParserDASH.cpp
+++ b/Engine/Plugins/Media/ElectraPlayer/Source/ElectraPlayerRuntime/Private/Runtime/Player/DASH/ManifestParserDASH.cpp
@@ -3156,7 +3156,7 @@ static bool ParseConditionalUIntAttribute(IDashMPDElement::FBoolUInt64& OutValue
 
 static bool ParseDuration(FTimeValue& OutValue, FManifestParserDASH* Builder, const TCHAR* AttributeName, const TCHAR* AttributeValue, int32 XmlLineNumber)
 {
-	if (ISO8601::ParseDuration(OutValue, AttributeValue) != UEMEDIA_ERROR_OK)
+	if (!ISO8601::ParseDuration(OutValue, AttributeValue))
 	{
 		return Builder->SetLastError(FString::Printf(TEXT("Failed to parse time attribute '@%s'=%s of element <%s> in line %d"), AttributeName, AttributeValue, *Builder->GetCurrentElementName(), XmlLineNumber), ERRCODE_DASH_MPD_PARSER_BAD_TIME_VALUE);
 	}
@@ -3165,7 +3165,7 @@ static bool ParseDuration(FTimeValue& OutValue, FManifestParserDASH* Builder, co
 
 static bool ParseDateTime(FTimeValue& OutValue, FManifestParserDASH* Builder, const TCHAR* AttributeName, const TCHAR* AttributeValue, int32 XmlLineNumber)
 {
-	if (ISO8601::ParseDateTime(OutValue, FString(AttributeValue)) != UEMEDIA_ERROR_OK)
+	if (!ISO8601::ParseDateTime(OutValue, FString(AttributeValue)))
 	{
 		return Builder->SetLastError(FString::Printf(TEXT("Failed to parse datetime attribute '@%s'=%s of element <%s> in line %d"), AttributeName, AttributeValue, *Builder->GetCurrentElementName(), XmlLineNumber), ERRCODE_DASH_MPD_PARSER_BAD_TIME_VALUE);
 	}
diff --git a/Engine/Plugins/Media/ElectraPlayer/Source/ElectraPlayerRuntime/Private/Runtime/Player/DASH/PlaylistReaderDASH.cpp b/Engine/Plugins/Media/ElectraPlayer/Source/ElectraPlayerRuntime/Private/Runtime/Player/DASH/PlaylistReaderDASH.cpp
--- a/Engine/Plugins/Media/ElectraPlayer/Source/ElectraPlayerRuntime/Private/Runtime/Player/DASH/PlaylistReaderDASH.cpp
+++ b/Engine/Plugins/Media/ElectraPlayer/Source/ElectraPlayerRuntime/Private/Runtime/Player/DASH/PlaylistReaderDASH.cpp
@@ -664,7 +664,7 @@ void FPlaylistReaderDASH::TriggerTimeSynchronization()
 				if (UTCTimings[i]->GetSchemeIdUri().Equals(DASH::Schemes::TimingSources::Scheme_urn_mpeg_dash_utc_direct2014))
 				{
 					FTimeValue DirectTime;
-					if (ISO8601::ParseDateTime(DirectTime, UTCTimings[i]->GetValue()) == UEMEDIA_ERROR_OK)
+					if (ISO8601::ParseDateTime(DirectTime, UTCTimings[i]->GetValue()))
 					{
 						PlayerSessionServices->GetSynchronizedUTCTime()->SetTime(DirectTime);
 						// Remove this timing element. It can not be used indefinitely and we best get rid of it now.
@@ -1138,8 +1138,8 @@ void FPlaylistReaderDASH::ManifestDownloadCompleted(FResourceLoadRequestPtr Requ
 			{
 				// Parse the header
 				FTimeValue DateFromHeader;
-				UEMediaError DateParseError = RFC7231::ParseDateTime(DateFromHeader, ConnInfo->ResponseHeaders[i].Value);
-				if (DateParseError == UEMEDIA_ERROR_OK && DateFromHeader.IsValid())
+				bool bDateParsedOk = RFC7231::ParseDateTime(DateFromHeader, ConnInfo->ResponseHeaders[i].Value);
+				if (bDateParsedOk && DateFromHeader.IsValid())
 				{
 					PlayerSessionServices->GetSynchronizedUTCTime()->SetTime(ConnInfo->RequestStartTime, DateFromHeader);
 				}
@@ -1412,8 +1412,8 @@ void FPlaylistReaderDASH::Timesync_httphead_Completed(FResourceLoadRequestPtr Re
 				{
 					// Parse the header
 					FTimeValue DateFromHeader;
-					UEMediaError DateParseError = RFC7231::ParseDateTime(DateFromHeader, ConnInfo->ResponseHeaders[i].Value);
-					if (DateParseError == UEMEDIA_ERROR_OK && DateFromHeader.IsValid())
+					bool bDateParsedOk = RFC7231::ParseDateTime(DateFromHeader, ConnInfo->ResponseHeaders[i].Value);
+					if (bDateParsedOk && DateFromHeader.IsValid())
 					{
 						PlayerSessionServices->GetSynchronizedUTCTime()->SetTime(DateFromHeader);
 					}
@@ -1438,7 +1438,7 @@ void FPlaylistReaderDASH::Timesync_httpxsdate_Completed(FResourceLoadRequestPtr
 			// Note: Yes, this is a bit weird, but the xs:dateTime format is actually exactly the same as ISO-8601
 			//       so it is not clear why there is a distinction made in the UTCTiming scheme.
 			//       We keep the different callback handlers here though for completeness sake.
-			if (ISO8601::ParseDateTime(NewTime, Response) == UEMEDIA_ERROR_OK)
+			if (ISO8601::ParseDateTime(NewTime, Response))
 			{
 				PlayerSessionServices->GetSynchronizedUTCTime()->SetTime(NewTime);
 			}
@@ -1463,7 +1463,7 @@ void FPlaylistReaderDASH::Timesync_httpiso_Completed(FResourceLoadRequestPtr Req
 		if (GetResponseString(Response, Request).IsOK())
 		{
 			FTimeValue NewTime;
-			if (ISO8601::ParseDateTime(NewTime, Response) == UEMEDIA_ERROR_OK)
+			if (ISO8601::ParseDateTime(NewTime, Response))
 			{
 				PlayerSessionServices->GetSynchronizedUTCTime()->SetTime(NewTime);
 			}
diff --git a/Engine/Plugins/Media/ElectraPlayer/Source/ElectraPlayerRuntime/Private/Runtime/Player/HLS/ManifestBuilderHLS.cpp b/Engine/Plugins/Media/ElectraPlayer/Source/ElectraPlayerRuntime/Private/Runtime/Player/HLS/ManifestBuilderHLS.cpp
--- a/Engine/Plugins/Media/ElectraPlayer/Source/ElectraPlayerRuntime/Private/Runtime/Player/HLS/ManifestBuilderHLS.cpp
+++ b/Engine/Plugins/Media/ElectraPlayer/Source/ElectraPlayerRuntime/Private/Runtime/Player/HLS/ManifestBuilderHLS.cpp
@@ -1542,7 +1542,7 @@ FErrorDetail FManifestBuilderHLS::UpdateFromVariantPlaylist(TSharedPtrTS<FManife
 				// Is there an EXT-X-PROGRAM-DATE-TIME ?
 				if (GetSegmentTagString(AttrValue, HLSPlaylistParser::ExtXProgramDateTime))
 				{
-					if (ISO8601::ParseDateTime(Seg.AbsoluteDateTime, AttrValue)== UEMEDIA_ERROR_OK)
+					if (ISO8601::ParseDateTime(Seg.AbsoluteDateTime, AttrValue))
 					{
 						if (nSeg)
 						{
diff --git a/Engine/Plugins/Media/ElectraPlayer/Source/ElectraPlayerRuntime/Private/Runtime/Player/HLS/PlaylistReaderHLS.cpp b/Engine/Plugins/Media/ElectraPlayer/Source/ElectraPlayerRuntime/Private/Runtime/Player/HLS/PlaylistReaderHLS.cpp
--- a/Engine/Plugins/Media/ElectraPlayer/Source/ElectraPlayerRuntime/Private/Runtime/Player/HLS/PlaylistReaderHLS.cpp
+++ b/Engine/Plugins/Media/ElectraPlayer/Source/ElectraPlayerRuntime/Private/Runtime/Player/HLS/PlaylistReaderHLS.cpp
@@ -1061,8 +1061,8 @@ FErrorDetail FPlaylistReaderHLS::ParsePlaylist(FPlaylistRequestPtr FromRequest)
 						FTimeValue ResponseAge = MEDIAutcTime::Current() - ConnInfo->RequestStartTime - FTimeValue().SetFromSeconds(ConnInfo->TimeUntilFirstByte);
 						// Parse the header
 						FTimeValue DateFromHeader;
-						UEMediaError DateParseError = RFC7231::ParseDateTime(DateFromHeader, ConnInfo->ResponseHeaders[i].Value);
-						if (DateParseError == UEMEDIA_ERROR_OK && DateFromHeader.IsValid())
+						bool bDateParsedOk = RFC7231::ParseDateTime(DateFromHeader, ConnInfo->ResponseHeaders[i].Value);
+						if (bDateParsedOk && DateFromHeader.IsValid())
 						{
 							PlayerSessionServices->GetSynchronizedUTCTime()->SetTime(DateFromHeader + ResponseAge);
 						}
diff --git a/Engine/Plugins/Media/ElectraPlayer/Source/ElectraPlayerRuntime/Private/Runtime/TimeUtilities.cpp b/Engine/Plugins/Media/ElectraPlayer/Source/ElectraPlayerRuntime/Private/Runtime/TimeUtilities.cpp
deleted file mode 100644
diff --git a/Engine/Plugins/Media/ElectraPlayer/Source/ElectraPlayerRuntime/Private/Runtime/Utilities/StringHelpers.cpp b/Engine/Plugins/Media/ElectraPlayer/Source/ElectraPlayerRuntime/Private/Runtime/Utilities/StringHelpers.cpp
deleted file mode 100644
diff --git a/Engine/Plugins/Media/ElectraPlayer/Source/ElectraPlayerRuntime/Private/Runtime/Utilities/StringHelpers.h b/Engine/Plugins/Media/ElectraPlayer/Source/ElectraPlayerRuntime/Private/Runtime/Utilities/StringHelpers.h
deleted file mode 100644
diff --git a/Engine/Plugins/Media/ElectraPlayer/Source/ElectraPlayerRuntime/Private/Runtime/Utilities/TimeUtilities.h b/Engine/Plugins/Media/ElectraPlayer/Source/ElectraPlayerRuntime/Private/Runtime/Utilities/TimeUtilities.h
deleted file mode 100644
diff --git a/Engine/Plugins/Media/ElectraPlayer/Source/ElectraPlayerRuntime/Private/Runtime/Utilities/URI.cpp b/Engine/Plugins/Media/ElectraPlayer/Source/ElectraPlayerRuntime/Private/Runtime/Utilities/URI.cpp
deleted file mode 100644
diff --git a/Engine/Plugins/Media/ElectraPlayer/Source/ElectraPlayerRuntime/Private/Runtime/Utilities/URI.h b/Engine/Plugins/Media/ElectraPlayer/Source/ElectraPlayerRuntime/Private/Runtime/Utilities/URI.h
deleted file mode 100644
diff --git a/Engine/Plugins/Media/ElectraPlayer/Source/ElectraPlayerRuntime/Private/Runtime/Utilities/URLParser.cpp b/Engine/Plugins/Media/ElectraPlayer/Source/ElectraPlayerRuntime/Private/Runtime/Utilities/URLParser.cpp
deleted file mode 100644
diff --git a/Engine/Plugins/Media/ElectraPlayer/Source/ElectraPlayerRuntime/Private/Runtime/Utilities/URLParser.h b/Engine/Plugins/Media/ElectraPlayer/Source/ElectraPlayerRuntime/Private/Runtime/Utilities/URLParser.h
deleted file mode 100644
diff --git a/Engine/Plugins/Media/ElectraUtil/Source/ElectraBase/Private/StringHelpers.cpp b/Engine/Plugins/Media/ElectraUtil/Source/ElectraBase/Private/StringHelpers.cpp
new file mode 100644
--- /dev/null
+++ b/Engine/Plugins/Media/ElectraUtil/Source/ElectraBase/Private/StringHelpers.cpp
@@ -0,0 +1,145 @@
+// Copyright Epic Games, Inc. All Rights Reserved.
+
+#include "Utilities/StringHelpers.h"
+
+namespace Electra
+{
+	namespace StringHelpers
+	{
+
+
+		int32 FindFirstOf(const FString& InString, const FString& SplitAt, int32 FirstPos)
+		{
+			if (SplitAt.Len() == 1)
+			{
+				// Speical version for only one split character
+				const TCHAR& FindMe = SplitAt[0];
+				for (int32 i = FirstPos; i < InString.Len(); ++i)
+				{
+					if (InString[i] == FindMe)
+					{
+						return i;
+					}
+				}
+			}
+			else
+			{
+				for (int32 i = FirstPos; i < InString.Len(); ++i)
+				{
+					for (int32 j = 0; j < SplitAt.Len(); ++j)
+					{
+						if (InString[i] == SplitAt[j])
+						{
+							return i;
+						}
+					}
+				}
+			}
+			return INDEX_NONE;
+		}
+
+
+		int32 FindFirstNotOf(const FString& InString, const FString& InNotOfChars, int32 FirstPos)
+		{
+			for (int32 i = FirstPos; i < InString.Len(); ++i)
+			{
+				bool bFoundCharFromNotOfChars = false;
+				for (int32 j = 0; j < InNotOfChars.Len(); ++j)
+				{
+					if (InString[i] == InNotOfChars[j])
+					{
+						// We found a character from the "NOT of" list. Check next...
+						bFoundCharFromNotOfChars = true;
+						break;
+					}
+				}
+				if (!bFoundCharFromNotOfChars)
+				{
+					// We did not find any of the characters. This is what we are looking for and return the index of the first "not of" character
+					return i;
+				}
+			}
+			return INDEX_NONE;
+		}
+
+
+		int32 FindLastNotOf(const FString& InString, const FString& InNotOfChars, int32 InStartPos)
+		{
+			InStartPos = FMath::Min(InStartPos, InString.Len() - 1);
+			for (int32 i = InStartPos; i >= 0; --i)
+			{
+				bool bFoundCharFromNotOfChars = false;
+				for (int32 j = 0; j < InNotOfChars.Len(); ++j)
+				{
+					if (InString[i] == InNotOfChars[j])
+					{
+						// We found a character from the "NOT of" list. Check next...
+						bFoundCharFromNotOfChars = true;
+						break;
+					}
+				}
+				if (!bFoundCharFromNotOfChars)
+				{
+					// We did not find any of the characters. This is what we are looking for and return the index of the first "not of" character
+					return i;
+				}
+			}
+			return INDEX_NONE;
+		}
+
+
+		void SplitByDelimiter(TArray<FString>& OutSplits, const FString& InString, const FString& SplitAt)
+		{
+			if (InString.Len())
+			{
+				int32 FirstPos = 0;
+				while (1)
+				{
+					int32 SplitPos = InString.Find(SplitAt, ESearchCase::IgnoreCase, ESearchDir::FromStart, FirstPos);
+					FString subs = InString.Mid(FirstPos, SplitPos == INDEX_NONE ? MAX_int32 : SplitPos - FirstPos);
+					if (subs.Len())
+					{
+						OutSplits.Push(subs);
+					}
+					FirstPos = SplitPos + SplitAt.Len();
+					if (SplitPos == INDEX_NONE || FirstPos >= InString.Len())
+					{
+						break;
+					}
+				}
+			}
+		}
+
+
+		bool StringEquals(const TCHAR * const s1, const TCHAR * const s2)
+		{ 
+			return FPlatformString::Strcmp(s1, s2) == 0; 
+		}
+
+		bool StringStartsWith(const TCHAR * const s1, const TCHAR * const s2, SIZE_T n)
+		{ 
+			return FPlatformString::Strncmp(s1, s2, n) == 0; 
+		}
+
+		void StringToArray(TArray<uint8>& OutArray, const FString& InString)
+		{
+			FTCHARToUTF8 cnv(*InString);
+			int32 Len = cnv.Length();
+			OutArray.AddUninitialized(Len);
+			FMemory::Memcpy(OutArray.GetData(), cnv.Get(), Len);
+		}
+
+		FString ArrayToString(const TArray<uint8>& InArray)
+		{
+			FUTF8ToTCHAR cnv((const ANSICHAR*)InArray.GetData(), InArray.Num());
+			FString UTF8Text(cnv.Length(), cnv.Get());
+			//return (TCHAR*)FUTF8ToTCHAR((const ANSICHAR*)InArray.GetData(), InArray.Num()).Get();
+			//FString UTF8Text(InArray.Num(), (TCHAR*)FUTF8ToTCHAR((const ANSICHAR*)InArray.GetData(), InArray.Num()).Get());
+			return MoveTemp(UTF8Text);
+		}
+
+
+	} // namespace StringHelpers
+} // namespace Electra
+
+
diff --git a/Engine/Plugins/Media/ElectraUtil/Source/ElectraBase/Private/TimeUtilities.cpp b/Engine/Plugins/Media/ElectraUtil/Source/ElectraBase/Private/TimeUtilities.cpp
new file mode 100644
--- /dev/null
+++ b/Engine/Plugins/Media/ElectraUtil/Source/ElectraBase/Private/TimeUtilities.cpp
@@ -0,0 +1,590 @@
+// Copyright Epic Games, Inc. All Rights Reserved.
+
+#include "Utilities/TimeUtilities.h"
+#include "Utilities/StringHelpers.h"
+
+namespace Electra
+{
+
+	namespace
+	{
+		struct FTimeComponents
+		{
+			FTimeComponents()
+				: Years(0)
+				, Months(0)
+				, Days(0)
+				, Hours(0)
+				, Minutes(0)
+				, Seconds(0)
+				, Milliseconds(0)
+				, TimezoneMinutes(0)
+			{
+			}
+			int32		Years;
+			int32		Months;
+			int32		Days;
+			int32		Hours;
+			int32		Minutes;
+			int32		Seconds;
+			int32		Milliseconds;
+			int32		TimezoneMinutes;
+
+			bool IsValidUTC() const
+			{
+				return Years >= 1970 && Years < 2200 /* arbitrary choice in the future */ && Months >= 1 && Months <= 12 && Days >= 1 && Days <= 31 &&
+					Hours >= 0 && Hours <= 23 && Minutes >= 0 && Minutes <= 59 && Seconds >= 0 && Seconds <= 60 /*allow for leap second*/ &&
+					Milliseconds >= 0 && Milliseconds <= 999 &&
+					TimezoneMinutes >= -14 * 60 && TimezoneMinutes <= 14 * 60 /* +/- 14 hours*/;
+			}
+
+			static bool IsDigit(const TCHAR c)
+			{
+				return c >= TCHAR('0') && c <= TCHAR('9');
+			}
+
+			static bool IsLeapYear(int32 InYear)
+			{
+				return (InYear % 4) == 0 && ((InYear % 100) != 0 || (InYear % 400) == 0);
+			};
+
+			FTimeValue ToUTC() const
+			{
+				static const uint8 DaysInMonths[2][12] = { {31, 28, 31, 30, 31, 30, 31, 31, 30, 31, 30, 31}, {31, 29, 31, 30, 31, 30, 31, 31, 30, 31, 30, 31} };
+				int64 Sum = 0;
+				for (int32 y = 1970; y < Years; ++y)
+				{
+					Sum += IsLeapYear(y) ? 366 : 365;
+				}
+				const uint8* const DaysPerMonth = DaysInMonths[IsLeapYear(Years) ? 1 : 0];
+				for (int32 m = 1; m < Months; ++m)
+				{
+					Sum += DaysPerMonth[m - 1];
+				}
+				Sum += Days - 1;
+				Sum *= 24;
+				Sum += Hours;
+				Sum *= 60;
+				Sum += Minutes;
+				Sum *= 60;
+				Sum += Seconds;
+				Sum -= TimezoneMinutes * 60;
+				Sum *= 1000;
+				Sum += Milliseconds;
+				return FTimeValue().SetFromMilliseconds(Sum);
+			}
+		};
+
+
+		static int64 ParseSubStringToInt(TCHAR* InFrom, int32 NumChars)
+		{
+			TCHAR* last = InFrom + NumChars;
+			TCHAR c = *last;
+			*last = TCHAR('\0');
+			int64 v;
+			LexFromString(v, InFrom);
+			*last = c;
+			return v;
+		}
+
+	}
+
+
+
+
+
+
+	namespace ISO8601
+	{
+
+		bool ParseDateTime(FTimeValue& OutTimeValue, const FString& DateTime)
+		{
+			// Is this a valid date time string in extended format (we require colons and dashes)
+			// 		YYYY-MM-DDTHH:MM:SS[.s*][Z]
+			if (DateTime.Len() >= 19 && DateTime[10] == TCHAR('T') && DateTime[4] == TCHAR('-') && DateTime[7] == TCHAR('-') && DateTime[13] == TCHAR(':') && DateTime[16] == TCHAR(':'))
+			{
+				// Make a mutable copy of the string we can alter.
+				FString tempString(DateTime);
+				TCHAR* tempBuf = GetData(tempString);
+
+				FTimeComponents TimeComponents;
+
+				if (FTimeComponents::IsDigit(tempBuf[0]) && FTimeComponents::IsDigit(tempBuf[1]) && FTimeComponents::IsDigit(tempBuf[2]) && FTimeComponents::IsDigit(tempBuf[3]))
+				{
+					TimeComponents.Years = ParseSubStringToInt(tempBuf, 4);
+					if (FTimeComponents::IsDigit(tempBuf[5]) && FTimeComponents::IsDigit(tempBuf[6]))
+					{
+						TimeComponents.Months = ParseSubStringToInt(tempBuf + 5, 2);
+						if (FTimeComponents::IsDigit(tempBuf[8]) && FTimeComponents::IsDigit(tempBuf[9]))
+						{
+							TimeComponents.Days = ParseSubStringToInt(tempBuf + 8, 2);
+							if (FTimeComponents::IsDigit(tempBuf[11]) && FTimeComponents::IsDigit(tempBuf[12]))
+							{
+								TimeComponents.Hours = ParseSubStringToInt(tempBuf + 11, 2);
+								if (FTimeComponents::IsDigit(tempBuf[14]) && FTimeComponents::IsDigit(tempBuf[15]))
+								{
+									TimeComponents.Minutes = ParseSubStringToInt(tempBuf + 14, 2);
+									if (FTimeComponents::IsDigit(tempBuf[17]) && FTimeComponents::IsDigit(tempBuf[18]))
+									{
+										TimeComponents.Seconds = ParseSubStringToInt(tempBuf + 17, 2);
+
+										int32 milliSeconds = 0;
+										// Are there fractions or timezone?
+										TCHAR* Suffix = tempBuf + 19;
+										if (Suffix[0] == TCHAR('.'))
+										{
+											uint32 secFractionalScale = 3;
+											for (++Suffix; FTimeComponents::IsDigit(*Suffix); ++Suffix)
+											{
+												// Only consider 3 fractional digits.
+												if (secFractionalScale)
+												{
+													--secFractionalScale;
+													milliSeconds = milliSeconds * 10 + (*Suffix - TCHAR('0'));
+												}
+											}
+											while (secFractionalScale)
+											{
+												--secFractionalScale;
+												milliSeconds *= 10;
+											}
+										}
+										TimeComponents.Milliseconds = milliSeconds;
+
+										// Time zone suffix?
+										if (Suffix[0] == TCHAR('+') || Suffix[0] == TCHAR('-'))
+										{
+											// Offsets have to be [+|-][hh[[:]mm]]
+											int32 tzLen = FCString::Strlen(Suffix);
+											TCHAR tzSign = Suffix[0];
+											if (tzLen >= 3)
+											{
+												int32 offM = 0;
+												int32 offH = ParseSubStringToInt(Suffix + 1, 2);
+												Suffix += 3;
+												tzLen -= 3;
+												// There may or may not be a colon to separate minutes from hours, or no minutes at all.
+												if (tzLen)
+												{
+													if (Suffix[0] == TCHAR(':'))
+													{
+														Suffix += 1;
+														tzLen -= 1;
+													}
+													// Minutes?
+													if (tzLen >= 2)
+													{
+														offM = ParseSubStringToInt(Suffix, 2);
+														Suffix += 2;
+													}
+												}
+												TimeComponents.TimezoneMinutes = offH * 60 + offM;
+												if (tzSign == TCHAR('-'))
+												{
+													TimeComponents.TimezoneMinutes = -TimeComponents.TimezoneMinutes;
+												}
+											}
+											else
+											{
+												return false;
+											}
+										}
+										else if (Suffix[0] == TCHAR('Z'))
+										{
+											// We treat the time as UTC in all cases. If it's already that, fine.
+											++Suffix;
+										}
+										if (Suffix[0] != TCHAR('\0'))
+										{
+											return false;
+										}
+									}
+								}
+							}
+						}
+					}
+				}
+
+				if (TimeComponents.IsValidUTC())
+				{
+					OutTimeValue = TimeComponents.ToUTC();
+					return true;
+				}
+			}
+			return false;
+		}
+
+
+		bool ParseDuration(FTimeValue& OutTimeValue, const TCHAR* InDuration)
+		{
+			// Parse an xs:duration element.
+			// We interpret 'years'/'months'/'days' as per the DASH-IF specification where
+			//    1 year = 12 months / 1 month = 30 days / 1 day = 24 hours / 1 hour = 60 minutes / 1 minute = 60 seconds.
+			// So 1 year is worth only 360 days !!!
+			FString tempString(InDuration);
+			// See https://www.w3schools.com/xml/schema_dtypes_date.asp
+			if ((tempString.Len() && InDuration[0] == TCHAR('P')) || (tempString.Len() > 1 && InDuration[0] == TCHAR('-') && InDuration[1] == TCHAR('P')))
+			{
+				// Make a mutable copy of the string we can alter.
+				TCHAR* tempBuf = GetData(tempString);
+
+				bool bHaveT = false;
+				bool bHaveY = false;
+				bool bHaveM = false;
+				bool bHaveD = false;
+				bool bHaveH = false;
+				bool bHaveS = false;
+				bool bIsNegative = InDuration[0] == TCHAR('-');
+				TCHAR* Start = bIsNegative ? tempBuf+2 : tempBuf+1;
+				TCHAR* Anchor = nullptr;
+				TCHAR* Decimal = nullptr;
+				int64 TotalSeconds = 0;
+				int64 HNSFraction = 0;
+				while(*Start)
+				{
+					// Skip over consecutive digits until we hit another char
+					if (FTimeComponents::IsDigit(*Start))
+					{
+						if (!Anchor)
+						{
+							Anchor = Start;
+						}
+						++Start;
+					}
+					// Found a 'T' separating year/month/day from hour/minute/second?
+					else if (*Start == TCHAR('T'))
+					{
+						if (!Anchor && !bHaveT)
+						{
+							bHaveT = true;
+							bHaveM = false;
+							++Start;
+						}
+						else
+						{
+							return false;
+						}
+					}
+					// Years
+					else if (*Start == TCHAR('Y'))
+					{
+						if (Anchor && !bHaveT && !bHaveY)
+						{
+							int64 Years = ParseSubStringToInt(Anchor, Start-Anchor);
+							TotalSeconds += Years * (12 * 30 * 24 * 60 * 60);
+							Anchor = nullptr;
+							++Start;
+							bHaveY = true;
+						}
+						else
+						{
+							return false;
+						}
+					}
+					// Months / Minutes
+					else if (*Start == TCHAR('M'))
+					{
+						if (Anchor && !bHaveM)
+						{
+							int64 Value = ParseSubStringToInt(Anchor, Start-Anchor);
+							if (!bHaveT)
+							{
+								TotalSeconds += Value * (30 * 24 * 60 * 60);
+							}
+							else
+							{
+								TotalSeconds += Value * 60;
+							}
+							Anchor = nullptr;
+							++Start;
+							bHaveM = true;
+						}
+						else
+						{
+							return false;
+						}
+					}
+					// Days
+					else if (*Start == TCHAR('D'))
+					{
+						if (Anchor && !bHaveT && !bHaveD)
+						{
+							int64 Days = ParseSubStringToInt(Anchor, Start-Anchor);
+							TotalSeconds += Days * (24 * 60 * 60);
+							Anchor = nullptr;
+							++Start;
+							bHaveD = true;
+						}
+						else
+						{
+							return false;
+						}
+					}
+					// Hours
+					else if (*Start == TCHAR('H'))
+					{
+						if (Anchor && bHaveT && !bHaveH)
+						{
+							int64 Hours = ParseSubStringToInt(Anchor, Start-Anchor);
+							TotalSeconds += Hours * (60 * 60);
+							Anchor = nullptr;
+							++Start;
+							bHaveH = true;
+						}
+						else
+						{
+							return false;
+						}
+					}
+					// Seconds
+					else if (*Start == TCHAR('S'))
+					{
+						if (Anchor && bHaveT && !bHaveS)
+						{
+							// Is this a decimal with fractions?
+							if (!Decimal)
+							{
+								int64 Seconds = ParseSubStringToInt(Anchor, Start-Anchor);
+								TotalSeconds += Seconds;
+							}
+							else
+							{
+								// First get the full seconds
+								int32 nd = Anchor-Decimal-1;
+								if (nd > 0)
+								{
+									int64 Seconds = ParseSubStringToInt(Decimal, nd);
+									TotalSeconds += Seconds;
+								}
+								// Now get the fractional part. We limit it to 7 digits as this is the resolution of an FTimeValue.
+								nd = Start - Anchor;
+								check(nd > 0);		// This would fire if there are no digits following the decimal point, but then Anchor would be nullptr and we wouldn't even get here.
+								if (nd > 7)
+								{
+									nd = 7;
+								}
+								HNSFraction = ParseSubStringToInt(Anchor, nd);
+								while(++nd <= 7)
+								{
+									HNSFraction *= 10;
+								}
+							}
+							Anchor = nullptr;
+							Decimal = nullptr;
+							++Start;
+							bHaveS = true;
+						}
+						else
+						{
+							return false;
+						}
+					}
+					// Decimal point? We allow both period and comma here.
+					else if (*Start == TCHAR('.') || *Start == TCHAR(','))
+					{
+						// Only allowed for seconds. We don't know what we are parsing yet.
+						if (bHaveT && !bHaveS && !Decimal)
+						{
+							// If there is no digit in front of the decimal point we allow this by setting the decimal to the delimiter.
+							Decimal = Anchor ? Anchor : Start;
+							Anchor = nullptr;
+							++Start;
+						}
+						else
+						{
+							return false;
+						}
+					}
+					else
+					{
+						return false;
+					}
+				}
+				if (Anchor == nullptr && Decimal == nullptr)
+				{
+					OutTimeValue.SetFromHNS(TotalSeconds * 10000000L + HNSFraction);
+					return true;
+				}
+			}
+			return false;
+		}
+
+
+	} // namespace ISO8601
+
+
+	/***************************************************************************************************************************************************/
+	/***************************************************************************************************************************************************/
+	/***************************************************************************************************************************************************/
+
+
+	namespace RFC7231
+	{
+
+		bool ParseDateTime(FTimeValue& OutTimeValue, const FString& DateTime)
+		{
+			static const TCHAR* const MonthNames[] = { TEXT("Jan"), TEXT("Feb"), TEXT("Mar"), TEXT("Apr"), TEXT("May"), TEXT("Jun"), TEXT("Jul"), TEXT("Aug"), TEXT("Sep"), TEXT("Oct"), TEXT("Nov"), TEXT("Dec") };
+
+			// Three formats need to be considered as per RFC 7231 section 7.1.1.1. :
+			//   First the preferred format, if that is no match then obsolete RFC 850 and finally ANSI C asctime()
+			// We split the string by space and comma to get 6 groups for IMF, 4 for RFC850 and 5 for asctime.
+			TCHAR Groups[6][16];
+			const TCHAR* s = GetData(DateTime);
+			int32 NumGroups = 0;
+			while (*s)
+			{
+				if (*s == TCHAR(' ') || *s == TCHAR(','))
+				{
+					++s;
+				}
+				else
+				{
+					TCHAR* Group = Groups[NumGroups];
+					if (++NumGroups > 6)
+					{
+						return false;
+					}
+					int32 GroupLen = 0;
+					while (*s && *s != TCHAR(' ') && *s != TCHAR(','))
+					{
+						*Group++ = *s++;
+						// Check that we do not exceed our fixed size array (including one to add terminating NUL char)
+						if (++GroupLen == sizeof(Groups[0]) - 1)
+						{
+							return false;
+						}
+					}
+					// Terminate the group.
+					*Group = TCHAR('\0');
+				}
+			}
+
+			FTimeComponents TimeComponents;
+
+			// Preferred format: IMF-fixdate, ie:  Sun, 06 Nov 1994 08:49:37 GMT
+			if (NumGroups == 6)
+			{
+				TimeComponents.Hours = ParseSubStringToInt(Groups[4], 2);
+				TimeComponents.Minutes = ParseSubStringToInt(Groups[4] + 3, 2);
+				TimeComponents.Seconds = ParseSubStringToInt(Groups[4] + 6, 2);
+				LexFromString(TimeComponents.Days, Groups[1]);
+				LexFromString(TimeComponents.Years, Groups[3]);
+				for (int32 Month = 0; Month < UE_ARRAY_COUNT(MonthNames); ++Month)
+				{
+					if (FCString::Strcmp(Groups[2], MonthNames[Month]) == 0)
+					{
+						TimeComponents.Months = Month + 1;
+						break;
+					}
+				}
+			}
+			// Obsolete RFC 850 format, ie:  Sunday, 06-Nov-94 08:49:37 GMT
+			else if (NumGroups == 4)
+			{
+				TimeComponents.Hours = ParseSubStringToInt(Groups[2], 2);
+				TimeComponents.Minutes = ParseSubStringToInt(Groups[2] + 3, 2);
+				TimeComponents.Seconds = ParseSubStringToInt(Groups[2] + 6, 2);
+
+				if (FCString::Strlen(Groups[1]) != 9)
+				{
+					return false;
+				}
+				TimeComponents.Days = ParseSubStringToInt(Groups[1], 2);
+				TimeComponents.Years = ParseSubStringToInt(Groups[1] + 7, 2);
+				// 1970-2069
+				TimeComponents.Years += TimeComponents.Years >= 70 ? 1900 : 2000;
+				for (int32 Month = 0; Month < UE_ARRAY_COUNT(MonthNames); ++Month)
+				{
+					if (FCString::Strncmp(Groups[1] + 3, MonthNames[Month], 3) == 0)
+					{
+						TimeComponents.Months = Month + 1;
+						break;
+					}
+				}
+			}
+			// ANSI C asctime() format, ie: Sun Nov  6 08:49:37 1994
+			else if (NumGroups == 5)
+			{
+				TimeComponents.Hours = ParseSubStringToInt(Groups[3], 2);
+				TimeComponents.Minutes = ParseSubStringToInt(Groups[3] + 3, 2);
+				TimeComponents.Seconds = ParseSubStringToInt(Groups[3] + 6, 2);
+				LexFromString(TimeComponents.Days, Groups[2]);
+				LexFromString(TimeComponents.Years, Groups[4]);
+				for (int32 Month = 0; Month < UE_ARRAY_COUNT(MonthNames); ++Month)
+				{
+					if (FCString::Strcmp(Groups[1], MonthNames[Month]) == 0)
+					{
+						TimeComponents.Months = Month + 1;
+						break;
+					}
+				}
+			}
+			else
+			{
+				return false;
+			}
+
+			if (TimeComponents.IsValidUTC())
+			{
+				OutTimeValue = TimeComponents.ToUTC();
+				return true;
+			}
+			return false;
+		}
+
+	} // namespace RFC7231
+
+
+	namespace RFC2326
+	{
+		bool ParseNPTTime(FTimeValue& OutTimeValue, const FString& NPTtime)
+		{
+			// See: https://www.w3.org/TR/media-frags/#naming-time
+			// and: https://www.ietf.org/rfc/rfc2326.txt  (section 3.6)
+			TArray<FString> NPTparts;
+			const TCHAR* const NPTDelimiter = TEXT(":");
+			NPTtime.ParseIntoArray(NPTparts, NPTDelimiter, false);
+			if (NPTparts.Num() < 4)
+			{
+				// Note that we do no validation here on whether or not the minutes and seconds are
+				// using double digits (in case of h:mm:ss.fff) or if they are within 0-59 range.
+				int64 h=0, m=0;
+				if (NPTparts.Num() == 3)
+				{
+					// h:mm:ss[.fff]
+					LexFromString(h, *NPTparts[0]);
+					LexFromString(m, *NPTparts[1]);
+				}
+				else if (NPTparts.Num() == 2)
+				{
+					// mm:ss[.fff]
+					LexFromString(m, *NPTparts[1]);
+				}
+				// The last part will always be seconds with optional fraction. We parse those independently from minutes and hours and add everything.
+				OutTimeValue = FTimeValue().SetFromMilliseconds((h * 3600 + m * 60) * 1000) + FTimeValue().SetFromTimeFraction(FTimeFraction().SetFromFloatString(NPTparts.Last()));
+				return true;
+			}
+			return false;
+		}
+	} // namespace RFC2326
+
+
+
+	namespace UnixEpoch
+	{
+		bool ParseFloatString(FTimeValue& OutTimeValue, const FString& Seconds)
+		{
+			FTimeFraction t;
+			if (t.SetFromFloatString(Seconds).IsValid())
+			{
+				OutTimeValue.SetFromTimeFraction(t);
+				return true;
+			}
+			return false;
+		}
+	}
+
+
+} // namespace Electra
+
+
diff --git a/Engine/Plugins/Media/ElectraUtil/Source/ElectraBase/Private/URLParser.cpp b/Engine/Plugins/Media/ElectraUtil/Source/ElectraBase/Private/URLParser.cpp
new file mode 100644
--- /dev/null
+++ b/Engine/Plugins/Media/ElectraUtil/Source/ElectraBase/Private/URLParser.cpp
@@ -0,0 +1,761 @@
+// Copyright Epic Games, Inc. All Rights Reserved.
+
+#include "Utilities/URLParser.h"
+
+
+namespace Electra
+{
+
+	static inline void _SwapStrings(FString& a, FString& b)
+	{
+		FString s = MoveTemp(a);
+		a = MoveTemp(b);
+		b = MoveTemp(s);
+	}
+
+	void FURL_RFC3986::Empty()
+	{
+		Scheme.Empty();
+		UserInfo.Empty();
+		Host.Empty();
+		Port.Empty();
+		Path.Empty();
+		Query.Empty();
+		Fragment.Empty();
+	}
+
+	void FURL_RFC3986::Swap(FURL_RFC3986& Other)
+	{
+		_SwapStrings(Scheme, Other.Scheme);
+		_SwapStrings(UserInfo, Other.UserInfo);
+		_SwapStrings(Host, Other.Host);
+		_SwapStrings(Port, Other.Port);
+		_SwapStrings(Path, Other.Path);
+		_SwapStrings(Query, Other.Query);
+		_SwapStrings(Fragment, Other.Fragment);
+	}
+
+	bool FURL_RFC3986::Parse(const FString& InURL)
+	{
+		Empty();
+		if (InURL.IsEmpty())
+		{
+			return true;
+		}
+		// Not handling URNs here.
+		if (InURL.StartsWith(TEXT("urn:")))
+		{
+			return false;
+		}
+		/*
+		RFC 3986, section 3: Syntax components
+
+				 foo://example.com:8042/over/there?name=ferret#nose
+				 \_/   \______________/\_________/ \_________/ \__/
+				  |           |            |            |        |
+			   scheme     authority       path        query   fragment
+		*/
+		StringHelpers::FStringIterator it(InURL);
+		if (InURL.StartsWith(TEXT("//")))
+		{
+			it += 2;
+			if (ParseAuthority(it))
+			{
+				return ParsePathAndQueryFragment(it);
+			}
+			return false;
+		}
+		// Check if the URL begins with a path, query or fragment.
+		else if (IsPathSeparator(*it) || IsQueryOrFragmentSeparator(*it) || *it == TCHAR('.'))
+		{
+			return ParsePathAndQueryFragment(it);
+		}
+		else
+		{
+			// If we get a URL without a scheme we will trip over the colon separating host and port, mistaking it for
+			// the scheme delimiter. For our purposes we only handle URLs that have a scheme when they have an authority!
+			FString Component;
+			// Parse out the first component up to the delimiting colon, path, query or fragment.
+			while(it && !IsColonSeparator(*it) && !IsQueryOrFragmentSeparator(*it) && !IsPathSeparator(*it))
+			{
+				Component += *it++;
+			}
+			// If we found the colon we assume it separates the scheme from the authority.
+			if (it && IsColonSeparator(*it))
+			{
+				++it;
+				if (!it)
+				{
+					// Cannot just be a scheme with no authority or path.
+					return false;
+				}
+				// Assume we parsed the scheme.
+				Scheme = Component;
+				// Does an absolute path ('/') or an authority ('//') follow?
+				if (IsPathSeparator(*it))
+				{
+					++it;
+					// Authority?
+					if (it && IsPathSeparator(*it))
+					{
+						++it;
+						// Parse out the authority and if successful continue parsing out the path.
+						if (!ParseAuthority(it))
+						{
+							return false;
+						}
+					}
+					else
+					{
+						// Unwind the absolute path '/' and parse the path.
+						--it;
+					}
+				}
+				if (!ParsePathAndQueryFragment(it))
+				{
+					return false;
+				}
+			}
+			else
+			{
+				// Got either '/', '?' or '#' while scanning for the scheme. Rewind the iterator and handle the entire thing as a path with or without query or fragment.
+				it.Reset();
+				return ParsePathAndQueryFragment(it);
+			}
+		}
+		return true;
+	}
+
+	FString FURL_RFC3986::GetScheme() const
+	{
+		return Scheme;
+	}
+
+	FString FURL_RFC3986::GetHost() const
+	{
+		return Host;
+	}
+
+	FString FURL_RFC3986::GetPath() const
+	{
+		return Path;
+	}
+
+	FString FURL_RFC3986::GetQuery() const
+	{
+		return Query;
+	}
+
+	FString FURL_RFC3986::GetFragment() const
+	{
+		return Fragment;
+	}
+
+	FString FURL_RFC3986::Get(bool bIncludeQuery, bool bIncludeFragment)
+	{
+		static const FString RequiredEscapeCharsPath(TEXT("?#"));
+
+		FString URL;
+		if (IsAbsolute())
+		{
+			FString Authority = GetAuthority();
+			URL = Scheme;
+			URL += TEXT(":");
+			if (Authority.Len() || Scheme.Equals(TEXT("file")))
+			{
+				URL += TEXT("//");
+				URL += Authority;
+			}
+			if (Path.Len())
+			{
+				if (Authority.Len() && !IsPathSeparator(Path[0]))
+				{
+					URL += TEXT("/");
+				}
+				UrlEncode(URL, Path, RequiredEscapeCharsPath);
+			}
+			else if ((Query.Len() && bIncludeQuery) || (Fragment.Len() && bIncludeFragment))
+			{
+				URL += TEXT("/");
+			}
+		}
+		else
+		{
+			UrlEncode(URL, Path, RequiredEscapeCharsPath);
+		}
+		if (Query.Len() && bIncludeQuery)
+		{
+			URL += TEXT("?");
+			URL += Query;
+		}
+		if (Fragment.Len() && bIncludeFragment)
+		{
+			URL += TEXT("#");
+			UrlEncode(URL, Fragment, FString());
+		}
+		return URL;
+	}
+
+	void FURL_RFC3986::GetQueryParams(TArray<FQueryParam>& OutQueryParams, const FString& InQueryParameters, bool bPerformUrlDecoding, bool bSameNameReplacesValue)
+	{
+		TArray<FString> ValuePairs;
+		InQueryParameters.ParseIntoArray(ValuePairs, TEXT("&"), true);
+		for(int32 i=0; i<ValuePairs.Num(); ++i)
+		{
+			int32 EqPos = 0;
+			ValuePairs[i].FindChar(TCHAR('='), EqPos);
+			FQueryParam qp;
+			bool bOk = true;
+			if (bPerformUrlDecoding)
+			{
+				bOk = UrlDecode(qp.Name, ValuePairs[i].Mid(0, EqPos)) && UrlDecode(qp.Value, ValuePairs[i].Mid(EqPos+1));
+			}
+			else
+			{
+				qp.Name = ValuePairs[i].Mid(0, EqPos);
+				qp.Value = ValuePairs[i].Mid(EqPos+1);
+			}
+			if (bOk)
+			{
+				if (!bSameNameReplacesValue)
+				{
+					OutQueryParams.Emplace(MoveTemp(qp));
+				}
+				else
+				{
+					bool bFound = false;
+					for(int32 j=0; j<OutQueryParams.Num(); ++j)
+					{
+						if (OutQueryParams[j].Name.Equals(qp.Name))
+						{
+							OutQueryParams[j] = MoveTemp(qp);
+							bFound = true;
+							break;
+						}
+					}
+					if (!bFound)
+					{
+						OutQueryParams.Emplace(MoveTemp(qp));
+					}
+				}
+			}
+		}
+	}
+
+	void FURL_RFC3986::GetQueryParams(TArray<FQueryParam>& OutQueryParams, bool bPerformUrlDecoding, bool bSameNameReplacesValue)
+	{
+		GetQueryParams(OutQueryParams, Query, bPerformUrlDecoding, bSameNameReplacesValue);
+	}
+
+	// Appends or prepends additional query parameters.
+	void FURL_RFC3986::AddQueryParameters(const FString& InQueryParameters, bool bAppend)
+	{
+		if (InQueryParameters.Len() > 0)
+		{
+			FString qp(InQueryParameters);
+			if (qp[0] == TCHAR('?') || qp[0] == TCHAR('&'))
+			{
+				qp.RightChopInline(1, false);
+			}
+			if (qp.Len() && qp[qp.Len()-1] == TCHAR('&'))
+			{
+				qp.LeftChopInline(1, false);
+			}
+			if (bAppend)
+			{
+				if (Query.Len())
+				{
+					Query.AppendChar(TCHAR('&'));
+				}
+				Query.Append(qp);
+			}
+			else
+			{
+				if (Query.Len())
+				{
+					Query = qp + TCHAR('&') + Query;
+				}
+				else
+				{
+					Query = qp;
+				}
+			}
+		}
+	}
+
+
+	bool FURL_RFC3986::IsAbsolute() const
+	{
+		return Scheme.Len() > 0;
+	}
+
+	void FURL_RFC3986::GetPathComponents(TArray<FString>& OutPathComponents) const
+	{
+		GetPathComponents(OutPathComponents, GetPath());
+	}
+
+	FString FURL_RFC3986::GetAuthority() const
+	{
+		FString Authority;
+		if (UserInfo.Len())
+		{
+			Authority += UserInfo;
+			Authority += TEXT("@");
+		}
+		if (Scheme.Equals(TEXT("file")))
+		{
+			Authority += Host;
+		}
+		else
+		{
+			// Is the host an IPv6 address that needs to be enclosed in square brackets?
+			int32 DummyPos = 0;
+			if (Host.FindChar(TCHAR(':'), DummyPos))
+			{
+				Authority += TEXT("[");
+				Authority += Host;
+				Authority += TEXT("]");
+			}
+			else
+			{
+				Authority += Host;
+			}
+			// Need to append a port?
+			if (Port.Len())
+			{
+				Authority += TEXT(":");
+				Authority += Port;
+			}
+		}
+		return Authority;
+	}
+
+
+	bool FURL_RFC3986::HasSameOriginAs(const FURL_RFC3986& Other)
+	{
+		/* RFC 6454:
+			5.  Comparing Origins
+
+			   Two origins are "the same" if, and only if, they are identical.  In
+			   particular:
+
+			   o  If the two origins are scheme/host/port triples, the two origins
+				  are the same if, and only if, they have identical schemes, hosts,
+				  and ports.
+		*/
+		return Scheme.Equals(Other.Scheme) && Host.Equals(Other.Host) && Port.Equals(Other.Port);
+	}
+
+
+	bool FURL_RFC3986::ParseAuthority(StringHelpers::FStringIterator& it)
+	{
+		UserInfo.Empty();
+		FString Component;
+		while(it && !IsPathSeparator(*it) && !IsQueryOrFragmentSeparator(*it))
+		{
+			// If there is a user-info delimiter?
+			if (*it == TCHAR('@'))
+			{
+				// Yes, what we have gathered so far is the user info. Set it and gather from scratch.
+				UserInfo = Component;
+				Component.Empty();
+			}
+			else
+			{
+				Component += *it;
+			}
+			++it;
+		}
+		StringHelpers::FStringIterator SubIt(Component);
+		return ParseHostAndPort(SubIt);
+	}
+
+	bool FURL_RFC3986::ParseHostAndPort(StringHelpers::FStringIterator& it)
+	{
+		if (!it)
+		{
+			return true;
+		}
+		Host.Empty();
+		if (Scheme.Equals(TEXT("file")))
+		{
+			// For file:// scheme we have to consider Windows drive letters with a colon, eg. file://D:/
+			// We must not stop parsing at the colon since it will not indicate a port for the file scheme anyway.
+			while(it)
+			{
+				Host += *it++;
+			}
+		}
+		else
+		{
+			// IPv6 adress in [xxx:xxx:xxx] notation?
+			if (*it == TCHAR('['))
+			{
+				++it;
+				while(it && *it != TCHAR(']'))
+				{
+					Host += *it++;
+				}
+				if (!it)
+				{
+					// Need to have at least the closing ']'
+					return false;
+				}
+				++it;
+			}
+			else
+			{
+				while(it && !IsColonSeparator(*it))
+				{
+					Host += *it++;
+				}
+			}
+			if (it && IsColonSeparator(*it))
+			{
+				++it;
+				Port.Empty();
+				while(it)
+				{
+					Port += *it++;
+				}
+			}
+		}
+		return true;
+	}
+
+	bool FURL_RFC3986::ParsePathAndQueryFragment(StringHelpers::FStringIterator& it)
+	{
+		if (it)
+		{
+			if (!IsQueryOrFragmentSeparator(*it))
+			{
+				if (!ParsePath(it))
+				{
+					return false;
+				}
+			}
+			if (it && IsQuerySeparator(*it))
+			{
+				if (!ParseQuery(++it))
+				{
+					return false;
+				}
+			}
+			if (it && IsFragmentSeparator(*it))
+			{
+				return ParseFragment(++it);
+			}
+		}
+		return true;
+	}
+
+	bool FURL_RFC3986::ParsePath(StringHelpers::FStringIterator& it)
+	{
+		Path.Empty();
+		while(it && !IsQueryOrFragmentSeparator(*it))
+		{
+			Path += *it++;
+		}
+		// URL decode the path internally.
+		FString Decoded;
+		if (UrlDecode(Decoded, Path))
+		{
+			Path = Decoded;
+			return true;
+		}
+		return false;
+	}
+
+	bool FURL_RFC3986::ParseQuery(StringHelpers::FStringIterator& it)
+	{
+		Query.Empty();
+		// Query is all up to the end or the start of the fragment.
+		while(it && !IsFragmentSeparator(*it))
+		{
+			Query += *it++;
+		}
+		return true;
+	}
+
+	bool FURL_RFC3986::ParseFragment(StringHelpers::FStringIterator& it)
+	{
+		// Fragment being the last part of the URL extends all the way to the end.
+		Fragment = it.GetRemainder();
+		it.SetToEnd();
+		// URL decode the fragment internally.
+		FString Decoded;
+		if (UrlDecode(Decoded, Fragment))
+		{
+			Fragment = Decoded;
+			return true;
+		}
+		return false;
+	}
+
+	bool FURL_RFC3986::UrlDecode(FString& OutResult, const FString& InUrlToDecode)
+	{
+		StringHelpers::FStringIterator it(InUrlToDecode);
+		while(it)
+		{
+			TCHAR c = *it++;
+			if (c == TCHAR('%'))
+			{
+				if (!it)
+				{
+					// '%' at the end with nothing following!
+					return false;
+				}
+				TCHAR hi = *it++;
+				if (!it)
+				{
+					// Only one char after '%'. Need two!!
+					return false;
+				}
+				TCHAR lo = *it++;
+				if (lo >= TCHAR('0') && lo <= TCHAR('9'))
+				{
+					lo -= TCHAR('0');
+				}
+				else if (lo >= TCHAR('a') && lo <= TCHAR('f'))
+				{
+					lo = lo - TCHAR('a') + 10;
+				}
+				else if (lo >= TCHAR('A') && lo <= TCHAR('F'))
+				{
+					lo = lo - TCHAR('A') + 10;
+				}
+				else
+				{
+					// Not a hex digit!
+					return false;
+				}
+				if (hi >= TCHAR('0') && hi <= TCHAR('9'))
+				{
+					hi -= TCHAR('0');
+				}
+				else if (hi >= TCHAR('a') && hi <= TCHAR('f'))
+				{
+					hi = hi - TCHAR('a') + 10;
+				}
+				else if (hi >= TCHAR('A') && hi <= TCHAR('F'))
+				{
+					hi = hi - TCHAR('A') + 10;
+				}
+				else
+				{
+					// Not a hex digit!
+					return false;
+				}
+
+				c = hi * 16 + lo;
+			}
+			OutResult += c;
+		}
+		return true;
+	}
+
+	bool FURL_RFC3986::UrlEncode(FString& OutResult, const FString& InUrlToEncode, const FString& InReservedChars)
+	{
+		static const FString RequiredEscapeChars(TEXT("%<>{}|\\\"^`"));
+		int32 DummyPos = 0;
+		for(StringHelpers::FStringIterator it(InUrlToEncode); it; ++it)
+		{
+			TCHAR c = *it;
+			if ((c >= TCHAR('a') && c <= TCHAR('z')) || (c >= TCHAR('0') && c <= TCHAR('9')) || (c >= TCHAR('A') && c <= TCHAR('Z')) || c == TCHAR('-') || c == TCHAR('_') || c == TCHAR('.') || c == TCHAR('~'))
+			{
+				OutResult += c;
+			}
+			else if (c <= 0x20 || c >= 0x7f || RequiredEscapeChars.FindChar(c, DummyPos) || InReservedChars.FindChar(c, DummyPos))
+			{
+				OutResult += FString::Printf(TEXT("%%%02X"), c);
+			}
+			else
+			{
+				OutResult += c;
+			}
+		}
+		return true;
+	}
+
+	void FURL_RFC3986::GetPathComponents(TArray<FString>& OutPathComponents, const FString& InPath)
+	{
+		TArray<FString> Components;
+		// Split on '/', ignoring resulting empty parts (eg. "a//b" gives ["a", "b"] only instead of ["a", "", "b"] !!
+		InPath.ParseIntoArray(Components, TEXT("/"), true);
+		OutPathComponents.Append(Components);
+	}
+
+	FURL_RFC3986& FURL_RFC3986::ResolveWith(const FString& InChildURL)
+	{
+		FURL_RFC3986 Other;
+		if (Other.Parse(InChildURL))
+		{
+			ResolveWith(Other);
+		}
+		return *this;
+	}
+
+	FURL_RFC3986& FURL_RFC3986::ResolveAgainst(const FString& InParentURL)
+	{
+		FURL_RFC3986 Parent;
+		if (Parent.Parse(InParentURL))
+		{
+			Parent.ResolveWith(*this);
+			Swap(Parent);
+		}
+		return *this;
+	}
+	
+	/**
+	 * Resolve URL as per RFC 3986 section 5.2
+	 */
+	void FURL_RFC3986::ResolveWith(const FURL_RFC3986& Other)
+	{
+		if (!Other.Scheme.IsEmpty())
+		{
+			Scheme   = Other.Scheme;
+			UserInfo = Other.UserInfo;
+			Host     = Other.Host;
+			Port     = Other.Port;
+			Path     = Other.Path;
+			Query    = Other.Query;
+			RemoveDotSegments();
+		}
+		else
+		{
+			if (!Other.Host.IsEmpty())
+			{
+				UserInfo = Other.UserInfo;
+				Host     = Other.Host;
+				Port     = Other.Port;
+				Path     = Other.Path;
+				Query    = Other.Query;
+				RemoveDotSegments();
+			}
+			else
+			{
+				if (Other.Path.IsEmpty())
+				{
+					if (!Other.Query.IsEmpty())
+					{
+						Query = Other.Query;
+					}
+				}
+				else
+				{
+					if (IsPathSeparator(Other.Path[0]))
+					{
+						Path = Other.Path;
+					}
+					else
+					{
+						MergePath(Other.Path);
+					}
+					RemoveDotSegments();
+					Query = Other.Query;
+				}
+			}
+		}
+		Fragment = Other.Fragment;
+	}
+
+	void FURL_RFC3986::BuildPathFromSegments(const TArray<FString>& Components, bool bAddLeadingSlash, bool bAddTrailingSlash)
+	{
+		Path.Empty();
+		int32 DummyPos = 0;
+		for(int32 i=0, iMax=Components.Num(); i<iMax; ++i)
+		{
+			if (i == 0)
+			{
+				if (bAddLeadingSlash)
+				{
+					Path = TEXT("/");
+				}
+				/*
+					As per RFC 3986 Section 4.2 Relative Reference:
+						A path segment that contains a colon character (e.g., "this:that")
+						cannot be used as the first segment of a relative-path reference, as
+						it would be mistaken for a scheme name.  Such a segment must be
+						preceded by a dot-segment (e.g., "./this:that") to make a relative-
+						path reference.
+				*/
+				else if (Scheme.IsEmpty() && Components[0].FindChar(TCHAR(':'), DummyPos))
+				{
+					Path = TEXT("./");
+				}
+			}
+			else
+			{
+				Path += TCHAR('/');
+			}
+			Path += Components[i];
+		}
+		if (bAddTrailingSlash)
+		{
+			Path += TCHAR('/');
+		}
+	}
+
+	void FURL_RFC3986::MergePath(const FString& InPathToMerge)
+	{
+		// RFC 3986 Section 5.2.3. Merge Paths
+		if (Host.Len() && Path.IsEmpty())
+		{
+			Path = TEXT("/");
+			Path += InPathToMerge;
+		}
+		else
+		{
+			int32 LastSlashPos = INDEX_NONE;
+			// Is there a '/' somewhere in the current path that's not at the end?
+			if (Path.FindLastChar(TCHAR('/'), LastSlashPos))
+			{
+				if (LastSlashPos != Path.Len()-1)
+				{
+					Path.LeftInline(LastSlashPos+1);
+				}
+				Path += InPathToMerge;
+			}
+			else
+			{
+				Path = InPathToMerge;
+			}
+		}
+	}
+
+	void FURL_RFC3986::RemoveDotSegments()
+	{
+		if (Path.Len())
+		{
+			TArray<FString> NewSegments;
+			// Remember the current path starting or ending in a slash.
+			bool bHasLeadingSlash = IsPathSeparator(Path[0]);
+			bool bHasTrailingSlash = Path.EndsWith(TEXT("/"));
+			// Split the path into its segments
+			TArray<FString> Segments;
+			GetPathComponents(Segments);
+			for(int32 i=0; i<Segments.Num(); ++i)
+			{
+				// For every ".." go up one level.
+				if (Segments[i].Equals(TEXT("..")))
+				{
+					if (NewSegments.Num())
+					{
+						NewSegments.Pop();
+					}
+				}
+				// Add non- "." segments.
+				else if (!Segments[i].Equals(TEXT(".")))
+				{
+					NewSegments.Add(Segments[i]);
+				}
+			}
+			BuildPathFromSegments(NewSegments, bHasLeadingSlash, bHasTrailingSlash);
+		}
+	}
+
+
+
+} // namespace Electra
+
+
diff --git a/Engine/Plugins/Media/ElectraUtil/Source/ElectraBase/Public/Utilities/StringHelpers.h b/Engine/Plugins/Media/ElectraUtil/Source/ElectraBase/Public/Utilities/StringHelpers.h
new file mode 100644
--- /dev/null
+++ b/Engine/Plugins/Media/ElectraUtil/Source/ElectraBase/Public/Utilities/StringHelpers.h
@@ -0,0 +1,146 @@
+// Copyright Epic Games, Inc. All Rights Reserved.
+
+#pragma once
+
+#include <CoreMinimal.h>
+
+
+namespace Electra
+{
+
+	inline void LexFromStringHex(int32& OutValue, const TCHAR* Buffer) { OutValue = FCString::Strtoi(Buffer, nullptr, 16); }
+	inline void LexFromStringHex64(int64& OutValue, const TCHAR* Buffer) { OutValue = FCString::Strtoi64(Buffer, nullptr, 16); }
+	inline void LexFromStringHexU64(uint64& OutValue, const TCHAR* Buffer) { OutValue = FCString::Strtoui64(Buffer, nullptr, 16); }
+
+	namespace StringHelpers
+	{
+
+		int32 ELECTRABASE_API FindFirstOf(const FString& InString, const FString& SplitAt, int32 FirstPos = 0);
+
+		int32 ELECTRABASE_API FindFirstNotOf(const FString& InString, const FString& InNotOfChars, int32 FirstPos = 0);
+
+		int32 ELECTRABASE_API FindLastNotOf(const FString& InString, const FString& InNotOfChars, int32 StartPos = MAX_int32);
+
+		void ELECTRABASE_API SplitByDelimiter(TArray<FString>& OutSplits, const FString& InString, const FString& SplitAt);
+
+		bool ELECTRABASE_API StringEquals(const TCHAR * const s1, const TCHAR * const s2);
+
+		bool ELECTRABASE_API StringStartsWith(const TCHAR * const s1, const TCHAR * const s2, SIZE_T n);
+
+		void ELECTRABASE_API StringToArray(TArray<uint8>& OutArray, const FString& InString);
+
+		FString ELECTRABASE_API ArrayToString(const TArray<uint8>& InArray);
+
+		/**
+		 * There is a known anomaly in the FString::TConstIterator. It iterates all TCHARs in the string *including* the terminating zero character.
+		 * This is not the behaviour we want and setup some helper iterator here which is not including the terminating zero.
+		 */
+		class ELECTRABASE_API FStringIterator
+		{
+		public:
+			FStringIterator(const FString& InString, int32 StartIndex = 0)
+				: StringToIterate(InString)
+				, Index(StartIndex)
+			{
+			}
+
+			/** Advances iterator to the next element in the container. */
+			FStringIterator& operator++()
+			{
+				++Index;
+				return *this;
+			}
+			FStringIterator operator++(int)
+			{
+				FStringIterator Tmp(*this);
+				++Index;
+				return Tmp;
+			}
+
+			/** Moves iterator to the previous element in the container. */
+			FStringIterator& operator--()
+			{
+				--Index;
+				return *this;
+			}
+			FStringIterator operator--(int)
+			{
+				FStringIterator Tmp(*this);
+				--Index;
+				return Tmp;
+			}
+
+			/** iterator arithmetic support */
+			FStringIterator& operator+=(int32 Offset)
+			{
+				Index += Offset;
+				return *this;
+			}
+
+			FStringIterator operator+(int32 Offset) const
+			{
+				FStringIterator Tmp(*this);
+				return Tmp += Offset;
+			}
+
+			FStringIterator& operator-=(int32 Offset)
+			{
+				return *this += -Offset;
+			}
+
+			FStringIterator operator-(int32 Offset) const
+			{
+				FStringIterator Tmp(*this);
+				return Tmp -= Offset;
+			}
+
+			const TCHAR& operator* () const
+			{
+				return StringToIterate[Index];
+			}
+
+			const TCHAR* operator->() const
+			{
+				return &StringToIterate[Index];
+			}
+
+			/** conversion to "bool" returning true if the iterator has not reached the last element. */
+			FORCEINLINE explicit operator bool() const
+			{
+				return Index >= 0 && Index < StringToIterate.Len();
+			}
+
+			const TCHAR* GetRemainder() const
+			{
+				return &StringToIterate[Index];
+			}
+
+			/** Returns an index to the current element. */
+			int32 GetIndex() const
+			{
+				return Index;
+			}
+
+			/** Resets the iterator to the first element. */
+			void Reset()
+			{
+				Index = 0;
+			}
+
+			/** Sets iterator to the last element. */
+			void SetToEnd()
+			{
+				Index = StringToIterate.Len();
+			}
+
+			FORCEINLINE friend bool operator==(const FStringIterator& Lhs, const FStringIterator& Rhs) { return &Lhs.StringToIterate == &Rhs.StringToIterate && Lhs.Index == Rhs.Index; }
+			FORCEINLINE friend bool operator!=(const FStringIterator& Lhs, const FStringIterator& Rhs) { return &Lhs.StringToIterate != &Rhs.StringToIterate || Lhs.Index != Rhs.Index; }
+
+		private:
+			const FString&	StringToIterate;
+			int32			Index;
+		};
+
+	} // namespace StringHelpers
+
+} // namespace Electra
diff --git a/Engine/Plugins/Media/ElectraUtil/Source/ElectraBase/Public/Utilities/TimeUtilities.h b/Engine/Plugins/Media/ElectraUtil/Source/ElectraBase/Public/Utilities/TimeUtilities.h
new file mode 100644
--- /dev/null
+++ b/Engine/Plugins/Media/ElectraUtil/Source/ElectraBase/Public/Utilities/TimeUtilities.h
@@ -0,0 +1,35 @@
+// Copyright Epic Games, Inc. All Rights Reserved.
+
+#pragma once
+
+#include <CoreMinimal.h>
+#include "PlayerTime.h"
+
+
+namespace Electra
+{
+
+	namespace ISO8601
+	{
+		bool ELECTRABASE_API ParseDateTime(FTimeValue& OutTimeValue, const FString& DateTime);
+		bool ELECTRABASE_API ParseDuration(FTimeValue& OutTimeValue, const TCHAR* InDuration);
+	}
+
+
+	namespace RFC7231
+	{
+		bool ELECTRABASE_API ParseDateTime(FTimeValue& OutTimeValue, const FString& DateTime);
+	}
+
+	namespace RFC2326
+	{
+		bool ELECTRABASE_API ParseNPTTime(FTimeValue& OutTimeValue, const FString& NPTtime);
+	}
+
+	namespace UnixEpoch
+	{
+		bool ELECTRABASE_API ParseFloatString(FTimeValue& OutTimeValue, const FString& Seconds);
+	}
+
+} // namespace Electra
+
diff --git a/Engine/Plugins/Media/ElectraUtil/Source/ElectraBase/Public/Utilities/URLParser.h b/Engine/Plugins/Media/ElectraUtil/Source/ElectraBase/Public/Utilities/URLParser.h
new file mode 100644
--- /dev/null
+++ b/Engine/Plugins/Media/ElectraUtil/Source/ElectraBase/Public/Utilities/URLParser.h
@@ -0,0 +1,102 @@
+// Copyright Epic Games, Inc. All Rights Reserved.
+
+#pragma once
+
+#include <CoreMinimal.h>
+#include "Utilities/StringHelpers.h"
+
+
+namespace Electra
+{
+
+	class ELECTRABASE_API FURL_RFC3986
+	{
+	public:
+		// Parses the given URL.
+		bool Parse(const FString& InURL);
+		// Returns the scheme, if present. Does not include the :// sequence.
+		FString GetScheme() const;
+		// Returns the host, if present.
+		FString GetHost() const;
+		// Returns the path. Escape sequences will already be decoded.
+		FString GetPath() const;
+		// Returns the entire query string, if any, without the leading '?'. Escape sequences will still be present!
+		FString GetQuery() const;
+		// Returns the fragment, if any, without the leading '#'. Escape sequences will already be decoded.
+		FString GetFragment() const;
+		// Returns the full URL with or without query and fragment parts. Characters will be escaped if necessary.
+		FString Get(bool bIncludeQuery=true, bool bIncludeFragment=true);
+		// Returns whether or not this URL is absolute (the scheme not being empty).
+		bool IsAbsolute() const;
+		// Returns the path as individual components. Like GetPath() the components will have escape sequences already decoded.
+		void GetPathComponents(TArray<FString>& OutPathComponents) const;
+		// Resolves a relative URL against this one.
+		FURL_RFC3986& ResolveWith(const FString& InChildURL);
+		// Resolves this URL (which should be relative) against the specified URL.
+		FURL_RFC3986& ResolveAgainst(const FString& InParentURL);
+
+		// Appends or prepends additional query parameters.
+		void AddQueryParameters(const FString& InQueryParameters, bool bAppend);
+
+		// Returns if this URL has the same origin as another one as per RFC 6454.
+		bool HasSameOriginAs(const FURL_RFC3986& Other);
+
+		struct FQueryParam
+		{
+			FString Name;
+			FString Value;
+		};
+		// Returns the query parameters as a list of name/value pairs.
+		void GetQueryParams(TArray<FQueryParam>& OutQueryParams, bool bPerformUrlDecoding, bool bSameNameReplacesValue=true);
+		// Returns the given query parameter string as a list of name/value pairs.
+		// The parameter string MUST NOT start with a '?'.
+		static void GetQueryParams(TArray<FQueryParam>& OutQueryParams, const FString& InQueryParameters, bool bPerformUrlDecoding, bool bSameNameReplacesValue=true);
+
+		// Decodes %XX escaped sequences into their original characters. Appends to the output. Hence in and out must not be the same.
+		static bool UrlDecode(FString& OutResult, const FString& InUrlToDecode);
+		// Encodes characters not permitted in a URL into %XX escaped sequences. Appends to the output. Hence in and out must not be the same.
+		static bool UrlEncode(FString& OutResult, const FString& InUrlToEncode, const FString& InReservedChars);
+
+	private:
+		FString Scheme;
+		FString UserInfo;
+		FString Host;
+		FString Port;
+		FString Path;
+		FString Query;
+		FString Fragment;
+
+		static void GetPathComponents(TArray<FString>& OutPathComponents, const FString& InPath);
+
+		void Empty();
+		void Swap(FURL_RFC3986& Other);
+		inline bool IsColonSeparator(TCHAR c)
+		{ return c == TCHAR(':'); }
+		inline bool IsPathSeparator(TCHAR c)
+		{ return c == TCHAR('/'); }
+		inline bool IsQuerySeparator(TCHAR c)
+		{ return c == TCHAR('?'); }
+		inline bool IsFragmentSeparator(TCHAR c)
+		{ return c == TCHAR('#'); }
+		inline bool IsQueryOrFragmentSeparator(TCHAR c)
+		{ return IsQuerySeparator(c) || IsFragmentSeparator(c); }
+
+		bool ParseAuthority(StringHelpers::FStringIterator& it);
+		bool ParseHostAndPort(StringHelpers::FStringIterator& it);
+		bool ParsePathAndQueryFragment(StringHelpers::FStringIterator& it);
+		bool ParsePath(StringHelpers::FStringIterator& it);
+		bool ParseQuery(StringHelpers::FStringIterator& it);
+		bool ParseFragment(StringHelpers::FStringIterator& it);
+
+		FString GetAuthority() const;
+		void ResolveWith(const FURL_RFC3986& Other);
+		void BuildPathFromSegments(const TArray<FString>& Components, bool bAddLeadingSlash, bool bAddTrailingSlash);
+		void MergePath(const FString& InPathToMerge);
+		void RemoveDotSegments();
+	};
+
+
+
+} // namespace Electra
+
+
