ElectraUtil: fixed Unix compilation errors

diff --git a/Engine/Plugins/Media/ElectraUtil/Source/ElectraSamples/Private/ElectraTextureSampleUnix.cpp b/Engine/Plugins/Media/ElectraUtil/Source/ElectraSamples/Private/ElectraTextureSampleUnix.cpp
--- a/Engine/Plugins/Media/ElectraUtil/Source/ElectraSamples/Private/ElectraTextureSampleUnix.cpp
+++ b/Engine/Plugins/Media/ElectraUtil/Source/ElectraSamples/Private/ElectraTextureSampleUnix.cpp
@@ -6,6 +6,8 @@
 
 #include "ElectraTextureSample.h"
 
+#if !UE_SERVER
+
 /**
 * Init code for realloacting an image from the the pool
 */
@@ -19,14 +21,9 @@ void FElectraTextureSampleUnix::InitializePoolable()
 */
 void FElectraTextureSampleUnix::ShutdownPoolable()
 {
-	TSharedPtr<MEDIArendererVideoUE, ESPMode::ThreadSafe> lockedVideoRenderer = OwningRenderer.Pin();
-	if (lockedVideoRenderer.IsValid())
-	{
-		lockedVideoRenderer->TextureReleasedToPool(GetDuration());
-	}
-
-	// This is coming from the original code: Drop reference to the texture. It should be released by the outside system.
 	Texture = nullptr;
 }
 
+#endif
+
 #endif
\ No newline at end of file
diff --git a/Engine/Plugins/Media/ElectraUtil/Source/ElectraSamples/Public/GStreamer/ElectraTextureSample.h b/Engine/Plugins/Media/ElectraUtil/Source/ElectraSamples/Public/GStreamer/ElectraTextureSample.h
--- a/Engine/Plugins/Media/ElectraUtil/Source/ElectraSamples/Public/GStreamer/ElectraTextureSample.h
+++ b/Engine/Plugins/Media/ElectraUtil/Source/ElectraSamples/Public/GStreamer/ElectraTextureSample.h
@@ -34,7 +34,7 @@ public:
 	{
 		check(IsInRenderingThread());
 
-		const uint32 CreateFlags = TexCreate_Dynamic | TexCreate_SRGB;
+		const ETextureCreateFlags CreateFlags = TexCreate_Dynamic | TexCreate_SRGB;
 
 		TRefCountPtr<FRHITexture2D> DummyTexture2DRHI;
 		FRHIResourceCreateInfo CreateInfo(TEXT("DummyTexture2D"));
@@ -126,9 +126,6 @@ public:
 		return Texture;
 	}
 
-	// We hold a weak reference to the video renderer. During destruction the video renderer could be destroyed while samples are still out there..
-	TWeakPtr<MEDIArendererVideoUE, ESPMode::ThreadSafe> OwningRenderer;
-
 private:
 	TRefCountPtr<FRHITexture2D> Texture;
 	FMediaTimeStamp Time;
