Fix issue with Adreno version checks

#android

diff --git a/Engine/Build/Android/Java/src/com/epicgames/unreal/MediaPlayer14.java b/Engine/Build/Android/Java/src/com/epicgames/unreal/MediaPlayer14.java
--- a/Engine/Build/Android/Java/src/com/epicgames/unreal/MediaPlayer14.java
+++ b/Engine/Build/Android/Java/src/com/epicgames/unreal/MediaPlayer14.java
@@ -884,11 +884,16 @@ public class MediaPlayer14
 				// Do not use shared context if Adreno before 400 or on older Android than Marshmallow
 				if (RendererString.contains("Adreno (TM) "))
 				{
-					int AdrenoVersion = Integer.parseInt(RendererString.substring(12));
-					if (AdrenoVersion < 400 || android.os.Build.VERSION.SDK_INT < 22)
+					java.util.regex.Pattern pattern = java.util.regex.Pattern.compile("Adreno \\(TM\\) (\\d*)");
+					java.util.regex.Matcher matcher = pattern.matcher(RendererString);
+					if (matcher.find() && matcher.groupCount() >= 1)
 					{
-						GameActivity.Log.debug("MediaPlayer14: disabled shared GL context on " + RendererString);
-						mUseOwnContext = false;
+						int AdrenoVersion = Integer.parseInt(matcher.group(1));
+						if (AdrenoVersion < 400 || android.os.Build.VERSION.SDK_INT < 22)
+						{
+							GameActivity.Log.debug("MediaPlayer14: disabled shared GL context on " + RendererString);
+							mUseOwnContext = false;
+						}
 					}
 				}
 			}
diff --git a/Engine/Plugins/Media/AndroidCamera/Source/AndroidCamera/Java/src/com/epicgames/unreal/CameraPlayer14.java b/Engine/Plugins/Media/AndroidCamera/Source/AndroidCamera/Java/src/com/epicgames/unreal/CameraPlayer14.java
--- a/Engine/Plugins/Media/AndroidCamera/Source/AndroidCamera/Java/src/com/epicgames/unreal/CameraPlayer14.java
+++ b/Engine/Plugins/Media/AndroidCamera/Source/AndroidCamera/Java/src/com/epicgames/unreal/CameraPlayer14.java
@@ -897,14 +897,19 @@ public class CameraPlayer14
 				// Do not use shared context if Adreno before 400 or on older Android than Marshmallow
 				if (RendererString.contains("Adreno (TM) "))
 				{
-					int AdrenoVersion = Integer.parseInt(RendererString.substring(12));
-					if (AdrenoVersion < 400 || android.os.Build.VERSION.SDK_INT < 22)
+					java.util.regex.Pattern pattern = java.util.regex.Pattern.compile("Adreno \\(TM\\) (\\d*)");
+					java.util.regex.Matcher matcher = pattern.matcher(RendererString);
+					if (matcher.find() && matcher.groupCount() >= 1)
 					{
-						GameActivity.Log.debug("MediaPlayer14: disabled shared GL context on " + RendererString);
-						mUseOwnContext = false;
+						int AdrenoVersion = Integer.parseInt(matcher.group(1));
+						if (AdrenoVersion < 400 || android.os.Build.VERSION.SDK_INT < 22)
+						{
+							GameActivity.Log.debug("CameraPlayer14: disabled shared GL context on " + RendererString);
+							mUseOwnContext = false;
+						}
 					}
 				}
-			}
+}			}
 
 			if (mUseOwnContext)
 			{
diff --git a/Engine/Plugins/Media/ElectraUtil/Source/ElectraSamples/Private/Java/src/com/epicgames/unreal/ElectraTextureSample.java b/Engine/Plugins/Media/ElectraUtil/Source/ElectraSamples/Private/Java/src/com/epicgames/unreal/ElectraTextureSample.java
--- a/Engine/Plugins/Media/ElectraUtil/Source/ElectraSamples/Private/Java/src/com/epicgames/unreal/ElectraTextureSample.java
+++ b/Engine/Plugins/Media/ElectraUtil/Source/ElectraSamples/Private/Java/src/com/epicgames/unreal/ElectraTextureSample.java
@@ -132,11 +132,16 @@ public class ElectraTextureSample
 				// Do not use shared context if Adreno before 400 or on older Android than Marshmallow
 				if (RendererString.contains("Adreno (TM) "))
 				{
-					int AdrenoVersion = Integer.parseInt(RendererString.substring(12));
-					if (AdrenoVersion < 400 || android.os.Build.VERSION.SDK_INT < 22)
+					java.util.regex.Pattern pattern = java.util.regex.Pattern.compile("Adreno \\(TM\\) (\\d*)");
+					java.util.regex.Matcher matcher = pattern.matcher(RendererString);
+					if (matcher.find() && matcher.groupCount() >= 1)
 					{
-						Log.d(TAG,"VideoDecoder: disabled shared GL context on " + RendererString);
-						mUseOwnContext = false;
+						int AdrenoVersion = Integer.parseInt(matcher.group(1));
+						if (AdrenoVersion < 400 || android.os.Build.VERSION.SDK_INT < 22)
+						{
+							GameActivity.Log.debug("VideoDecoder: disabled shared GL context on " + RendererString);
+							mUseOwnContext = false;
+						}
 					}
 				}
 			}
