ElectraHTTPStream: Added a network module geared towards downloading of chunked media streams for use with the Electra media player.

diff --git a/Engine/Plugins/Media/ElectraPlayer/Source/ElectraPlayerRuntime/Private/Runtime/StreamDataBuffer.h b/Engine/Plugins/Media/ElectraPlayer/Source/ElectraPlayerRuntime/Private/Runtime/StreamDataBuffer.h
--- a/Engine/Plugins/Media/ElectraPlayer/Source/ElectraPlayerRuntime/Private/Runtime/StreamDataBuffer.h
+++ b/Engine/Plugins/Media/ElectraPlayer/Source/ElectraPlayerRuntime/Private/Runtime/StreamDataBuffer.h
@@ -162,7 +162,7 @@ namespace Electra
 
 
 		//! Inserts elements into the buffer. Returns true if successful, false if there is no room.
-		bool PushData(const uint8* InData, int32 NumElements)
+		bool PushData(const uint8* InData, int32 NumElements, bool bRingbufferPush=true)
 		{
 			FMediaCriticalSection::ScopedLock Lock(AccessLock);
 			check(!bEOD);
@@ -179,7 +179,10 @@ namespace Electra
 					CopyData(WritePos, InData, NumElements);
 					if ((WritePos += NumElements) == DataEnd)
 					{
-						WritePos = Data;
+						if (bRingbufferPush)
+						{
+							WritePos = Data;
+						}
 					}
 				}
 				else
@@ -384,15 +387,15 @@ namespace Electra
 
 		FMediaCriticalSection	AccessLock;
 		FMediaEvent				SizeAvailableSignal;	//!< signaled when WaitingForSize data is present
-		uint8*					Data;					//!< Base address of buffer
-		uint8*					DataEnd;				//!< End address of buffer
-		uint8*					WritePos;				//!< Current write position
-		uint8*					ReadPos;				//!< Current read position
-		int32					DataSize;				//!< Maximum number of bytes in the buffer
-		int32					NumIn;					//!< Current number of bytes in the buffer
-		int32					WaitingForSize;			//!< If waiting for a certain number of bytes to become available
-		bool					bEOD;					//!< true when the last packet of data was pushed into the buffer.
-		bool					bWasAborted;
+		uint8* volatile			Data;					//!< Base address of buffer
+		uint8* volatile			DataEnd;				//!< End address of buffer
+		uint8* volatile			WritePos;				//!< Current write position
+		uint8* volatile			ReadPos;				//!< Current read position
+		volatile int32			DataSize;				//!< Maximum number of bytes in the buffer
+		volatile int32			NumIn;					//!< Current number of bytes in the buffer
+		volatile int32			WaitingForSize;			//!< If waiting for a certain number of bytes to become available
+		volatile bool			bEOD;					//!< true when the last packet of data was pushed into the buffer.
+		volatile bool			bWasAborted;
 	};
 
 
diff --git a/Engine/Plugins/Media/ElectraUtil/ElectraUtil.uplugin b/Engine/Plugins/Media/ElectraUtil/ElectraUtil.uplugin
--- a/Engine/Plugins/Media/ElectraUtil/ElectraUtil.uplugin
+++ b/Engine/Plugins/Media/ElectraUtil/ElectraUtil.uplugin
@@ -28,6 +28,13 @@
 			"LoadingPhase": "PostEngineInit",
 			"PlatformAllowList": [ "Win64", "Mac", "IOS", "Android" ],
 			"TargetDenyList": [ "Server" ]
+		},
+		{
+			"Name": "ElectraHTTPStream",
+			"Type": "RuntimeNoCommandlet",
+			"LoadingPhase": "PostEngineInit",
+			"PlatformAllowList": [ "Win64", "Mac", "IOS", "Android" ],
+			"TargetDenyList": [ "Server" ]
 		}
 	]
 }
\ No newline at end of file
diff --git a/Engine/Plugins/Media/ElectraUtil/Source/ElectraBase/Private/URLParser.cpp b/Engine/Plugins/Media/ElectraUtil/Source/ElectraBase/Private/URLParser.cpp
--- a/Engine/Plugins/Media/ElectraUtil/Source/ElectraBase/Private/URLParser.cpp
+++ b/Engine/Plugins/Media/ElectraUtil/Source/ElectraBase/Private/URLParser.cpp
@@ -136,6 +136,11 @@ namespace Electra
 		return Host;
 	}
 
+	FString FURL_RFC3986::GetPort() const
+	{
+		return Port;
+	}
+
 	FString FURL_RFC3986::GetPath() const
 	{
 		return Path;
@@ -196,6 +201,43 @@ namespace Electra
 		return URL;
 	}
 
+	FString FURL_RFC3986::GetPath(bool bIncludeQuery, bool bIncludeFragment)
+	{
+		static const FString RequiredEscapeCharsPath(TEXT("?#"));
+
+		FString URL;
+		if (IsAbsolute())
+		{
+			if (Path.Len())
+			{
+				if (GetAuthority().Len() && !IsPathSeparator(Path[0]))
+				{
+					URL += TEXT("/");
+				}
+				UrlEncode(URL, Path, RequiredEscapeCharsPath);
+			}
+			else if ((Query.Len() && bIncludeQuery) || (Fragment.Len() && bIncludeFragment))
+			{
+				URL += TEXT("/");
+			}
+		}
+		else
+		{
+			UrlEncode(URL, Path, RequiredEscapeCharsPath);
+		}
+		if (Query.Len() && bIncludeQuery)
+		{
+			URL += TEXT("?");
+			URL += Query;
+		}
+		if (Fragment.Len() && bIncludeFragment)
+		{
+			URL += TEXT("#");
+			UrlEncode(URL, Fragment, FString());
+		}
+		return URL;
+	}
+
 	void FURL_RFC3986::GetQueryParams(TArray<FQueryParam>& OutQueryParams, const FString& InQueryParameters, bool bPerformUrlDecoding, bool bSameNameReplacesValue)
 	{
 		TArray<FString> ValuePairs;
@@ -577,6 +619,19 @@ namespace Electra
 		return true;
 	}
 
+	FString FURL_RFC3986::GetStandardPortForScheme(const FString& InScheme, bool bIgnoreCase)
+	{
+		if (InScheme.Equals(TEXT("http"), bIgnoreCase ? ESearchCase::IgnoreCase : ESearchCase::CaseSensitive))
+		{
+			return FString(TEXT("80"));
+		}
+		else if (InScheme.Equals(TEXT("https"), bIgnoreCase ? ESearchCase::IgnoreCase : ESearchCase::CaseSensitive))
+		{
+			return FString(TEXT("443"));
+		}
+		return FString();
+	}
+
 	void FURL_RFC3986::GetPathComponents(TArray<FString>& OutPathComponents, const FString& InPath)
 	{
 		TArray<FString> Components;
diff --git a/Engine/Plugins/Media/ElectraUtil/Source/ElectraBase/Public/Utilities/URLParser.h b/Engine/Plugins/Media/ElectraUtil/Source/ElectraBase/Public/Utilities/URLParser.h
--- a/Engine/Plugins/Media/ElectraUtil/Source/ElectraBase/Public/Utilities/URLParser.h
+++ b/Engine/Plugins/Media/ElectraUtil/Source/ElectraBase/Public/Utilities/URLParser.h
@@ -18,6 +18,8 @@ namespace Electra
 		FString GetScheme() const;
 		// Returns the host, if present.
 		FString GetHost() const;
+		// Returns the port, if present.
+		FString GetPort() const;
 		// Returns the path. Escape sequences will already be decoded.
 		FString GetPath() const;
 		// Returns the entire query string, if any, without the leading '?'. Escape sequences will still be present!
@@ -26,6 +28,8 @@ namespace Electra
 		FString GetFragment() const;
 		// Returns the full URL with or without query and fragment parts. Characters will be escaped if necessary.
 		FString Get(bool bIncludeQuery=true, bool bIncludeFragment=true);
+		// Returns the path (no scheme or host) with or without query and fragment parts. Characters will be escaped if necessary.
+		FString GetPath(bool bIncludeQuery, bool bIncludeFragment);
 		// Returns whether or not this URL is absolute (the scheme not being empty).
 		bool IsAbsolute() const;
 		// Returns the path as individual components. Like GetPath() the components will have escape sequences already decoded.
@@ -57,6 +61,9 @@ namespace Electra
 		// Encodes characters not permitted in a URL into %XX escaped sequences. Appends to the output. Hence in and out must not be the same.
 		static bool UrlEncode(FString& OutResult, const FString& InUrlToEncode, const FString& InReservedChars);
 
+		// Returns the standard port for the given scheme. An empty string is returned if none is known.
+		static FString GetStandardPortForScheme(const FString& InScheme, bool bIgnoreCase=true);
+
 	private:
 		FString Scheme;
 		FString UserInfo;
diff --git a/Engine/Plugins/Media/ElectraUtil/Source/ElectraHTTPStream/ElectraHTTPStream.Build.cs b/Engine/Plugins/Media/ElectraUtil/Source/ElectraHTTPStream/ElectraHTTPStream.Build.cs
new file mode 100644
--- /dev/null
+++ b/Engine/Plugins/Media/ElectraUtil/Source/ElectraHTTPStream/ElectraHTTPStream.Build.cs
@@ -0,0 +1,157 @@
+// Copyright Epic Games, Inc. All Rights Reserved.
+
+using UnrealBuildTool;
+
+public class ElectraHTTPStream : ModuleRules
+{
+	// Returns true when this platform uses WinHttp
+	protected virtual bool bPlatformUsesWinHttp
+	{
+		get
+		{
+			return Target.Platform.IsInGroup(UnrealPlatformGroup.Windows);
+		}
+	}
+
+	// Returns whether or not this platform uses libCurl
+	protected virtual bool bPlatformUsesLibCurl
+	{
+		get
+		{
+			return Target.IsInPlatformGroup(UnrealPlatformGroup.Unix) || Target.IsInPlatformGroup(UnrealPlatformGroup.Android);
+		}
+	}
+
+
+	// Returns true if this platform supports xCurl (from Microsoft Game Core)
+	protected virtual bool bPlatformUsesXCurl
+	{
+		get
+		{
+			return false;
+		}
+	}
+
+
+	// Returns true if the platform requires all connections to be secure https:// connections.
+	protected virtual bool bPlatformRequiresSecureConnections
+	{
+		get
+		{
+			return false;
+		}
+	}
+
+	// Returns whether or not sockets are required for whatever HTTP library is being used.
+	protected virtual bool bPlatformRequiresSockets
+	{
+		get
+		{
+			return bPlatformUsesLibCurl;
+		}
+	}
+
+	// Returns whether or not this platform needs OpenSSL in combination with whatever HTTP library it is using.
+	protected virtual bool bPlatformRequiresOpenSSL
+	{
+		get
+		{
+			// WinHttp does not need OpenSSL.
+			if (bPlatformUsesWinHttp)
+			{
+				return false;
+			}
+			// libCurl is compiled to rely on OpenSSL
+			if (bPlatformUsesLibCurl)
+			{
+				return true;
+			}
+			return false;
+		}
+	}
+
+	public ElectraHTTPStream(ReadOnlyTargetRules Target) : base(Target)
+	{
+		PublicDefinitions.Add("ELECTRA_HTTPSTREAM_PACKAGE=1");
+
+		PrivateIncludePaths.AddRange(
+			new string[] {
+				"ElectraHTTPStream/Private",
+			});
+
+		PublicDependencyModuleNames.AddRange(
+			new string[] {
+				"Core"
+			});
+
+		PrivateDependencyModuleNames.AddRange(
+			new string[] {
+				"ElectraBase",
+				"SSL",
+			});
+
+		// When Curl is used we need to make sure there is no conflict with its use in the engine's HTTP module
+		// as far as global initialization is concerned.
+		if (bPlatformUsesLibCurl || bPlatformUsesXCurl)
+		{
+			PrivateDependencyModuleNames.AddRange(
+				new string[] {
+					"HTTP",
+				});
+		}
+
+		if (bPlatformRequiresSockets)
+		{
+			PrivateDependencyModuleNames.AddRange(
+				new string[] {
+					"Sockets",
+				}
+			);
+		}
+
+		if (bPlatformUsesWinHttp)
+		{
+			PublicDefinitions.Add("ELECTRA_HTTPSTREAM_WINHTTP=1");
+			AddEngineThirdPartyPrivateStaticDependencies(Target, "WinHttp");
+		}
+		else
+		{
+			PublicDefinitions.Add("ELECTRA_HTTPSTREAM_WINHTTP=0");
+		}
+
+		if (bPlatformUsesLibCurl || bPlatformUsesXCurl)
+		{
+			if (bPlatformUsesLibCurl)
+			{
+				PublicDefinitions.Add("ELECTRA_HTTPSTREAM_LIBCURL=1");
+				PublicDefinitions.Add("ELECTRA_HTTPSTREAM_XCURL=0");
+				AddEngineThirdPartyPrivateStaticDependencies(Target, "libcurl");
+			}
+			else
+			{
+				PublicDefinitions.Add("ELECTRA_HTTPSTREAM_LIBCURL=1");
+				PublicDefinitions.Add("ELECTRA_HTTPSTREAM_XCURL=1");
+			}
+		}
+		else
+		{
+			PublicDefinitions.Add("ELECTRA_HTTPSTREAM_LIBCURL=0");
+			PublicDefinitions.Add("ELECTRA_HTTPSTREAM_XCURL=0");
+		}
+
+		PublicDefinitions.Add("ELECTRA_HTTPSTREAM_REQUIRES_SECURE_CONNECTIONS=" + (bPlatformRequiresSecureConnections ? "1" : "0"));
+
+		if (bPlatformRequiresOpenSSL)
+		{
+			AddEngineThirdPartyPrivateStaticDependencies(Target, "OpenSSL");
+		}
+
+		// Apple
+		if (Target.Platform == UnrealTargetPlatform.IOS || Target.Platform == UnrealTargetPlatform.TVOS || Target.Platform == UnrealTargetPlatform.Mac)
+		{
+			PublicDefinitions.Add("ELECTRA_HTTPSTREAM_APPLE=1");
+			PublicFrameworks.Add("Security");
+		}
+
+	}
+}
diff --git a/Engine/Plugins/Media/ElectraUtil/Source/ElectraHTTPStream/Private/Apple/AppleElectraHTTPStream.cpp b/Engine/Plugins/Media/ElectraUtil/Source/ElectraHTTPStream/Private/Apple/AppleElectraHTTPStream.cpp
new file mode 100644
--- /dev/null
+++ b/Engine/Plugins/Media/ElectraUtil/Source/ElectraHTTPStream/Private/Apple/AppleElectraHTTPStream.cpp
@@ -0,0 +1,1465 @@
+// Copyright Epic Games, Inc. All Rights Reserved.
+
+#if ELECTRA_HTTPSTREAM_APPLE
+
+#include "Apple/AppleElectraHTTPStream.h"
+#include "HAL/ThreadSafeCounter.h"
+#include "HAL/PlatformTime.h"
+#include "HAL/PlatformProcess.h"
+#include "HAL/Runnable.h"
+#include "HAL/RunnableThread.h"
+#include "HAL/LowLevelMemTracker.h"
+
+#include "ElectraHTTPStreamModule.h"
+#include "ElectraHTTPStream.h"
+#include "ElectraHTTPStreamBuffer.h"
+#include "ElectraHTTPStreamResponse.h"
+#include "ElectraHTTPStreamPerf.h"
+
+#include "Utilities/TimeWaitableSignal.h"
+#include "Utilities/DefaultHttpUserAgent.h"
+#include "Utilities/HttpRangeHeader.h"
+#include "Utilities/URLParser.h"
+
+#include "Apple/AppleLogError.h"
+
+#if WITH_SSL
+#include "CommonCrypto/CommonDigest.h"
+#include "Apple/CFRef.h"
+#include "Ssl.h"
+#endif
+
+#include <atomic>
+
+
+// In all but shipping builds we allow disabling of security checks.
+#if !UE_BUILD_SHIPPING
+#define ELECTRA_HTTPSTREAM_APPLE_ALLOW_UNSAFE_CONNECTIONS_FOR_DEBUGGING 1
+#else
+#define ELECTRA_HTTPSTREAM_APPLE_ALLOW_UNSAFE_CONNECTIONS_FOR_DEBUGGING 0
+#endif
+
+#define ELECTRA_HTTPSTREAM_APPLE_ENABLE_SESSION_CHALLENGE !ELECTRA_HTTPSTREAM_APPLE_ALLOW_UNSAFE_CONNECTIONS_FOR_DEBUGGING
+
+
+class FElectraHTTPStreamRequestApple;
+@class FElectraHTTPStreamAppleSessionDelegate;
+
+/***************************************************************************************************************************************************/
+/***************************************************************************************************************************************************/
+/***************************************************************************************************************************************************/
+
+DECLARE_STATS_GROUP(TEXT("Electra HTTP Stream"), STATGROUP_ElectraHTTPStream, STATCAT_Advanced);
+DECLARE_CYCLE_STAT(TEXT("Process"), STAT_ElectraHTTPThread_Process, STATGROUP_ElectraHTTPStream);
+DECLARE_CYCLE_STAT(TEXT("Handle"), STAT_ElectraHTTPThread_EventHandle, STATGROUP_ElectraHTTPStream);
+DECLARE_CYCLE_STAT(TEXT("Custom handler"), STAT_ElectraHTTPThread_CustomHandler, STATGROUP_ElectraHTTPStream);
+
+/***************************************************************************************************************************************************/
+/***************************************************************************************************************************************************/
+/***************************************************************************************************************************************************/
+
+/**
+ * Apple version of the ElectraHTTPStream implementation.
+ */
+class FElectraHTTPStreamApple : public TSharedFromThis<FElectraHTTPStreamApple, ESPMode::ThreadSafe>, public IElectraHTTPStream, private FRunnable
+{
+public:
+	virtual ~FElectraHTTPStreamApple();
+
+	FElectraHTTPStreamApple();
+	virtual bool Initialize(const Electra::FParamDict& InOptions);
+
+	virtual FElectraHTTPStreamThreadHandlerDelegate& ThreadHandlerDelegate() override
+	{ return ThreadHandlerCallback; }
+
+	virtual void Close() override;
+
+	virtual IElectraHTTPStreamRequestPtr CreateRequest() override;
+
+	virtual void AddRequest(IElectraHTTPStreamRequestPtr Request) override;
+
+	NSURLSession* GetSession()
+	{ return Session; }
+
+	void TriggerWorkSignal()
+	{ HaveWorkSignal.Signal(); }
+
+	static TSharedPtr<FElectraHTTPStreamApple, ESPMode::ThreadSafe> GetHandlerForSession(NSURLSession* InSession)
+	{
+		FScopeLock lock(&AllSessionLock);
+		return AllSessions.Contains(InSession) ? AllSessions[InSession].Pin() : nullptr;
+	}
+
+	static TSharedPtr<FElectraHTTPStreamRequestApple, ESPMode::ThreadSafe> GetTask(NSURLSession* InSession, NSURLSessionTask* InTask)
+	{
+		TSharedPtr<FElectraHTTPStreamApple, ESPMode::ThreadSafe> Handler = GetHandlerForSession(InSession);
+		return Handler.IsValid() ? Handler->GetActiveTask(InTask) : nullptr;
+	}
+
+#if WITH_SSL
+	struct FChallengeResponse
+	{
+		NSURLSessionAuthChallengeDisposition Disposition = NSURLSessionAuthChallengeCancelAuthenticationChallenge;
+		NSURLCredential* Credential = nil;
+	};
+	FChallengeResponse ReceivedChallenge(NSURLAuthenticationChallenge* InChallenge);
+#endif
+
+private:
+	// Methods from FRunnable
+	virtual uint32 Run() override final;
+	virtual void Stop() override final;
+
+	void SetupNewRequests();
+	void UpdateActiveRequests();
+	void HandleCurl();
+	void HandleCompletedRequests();
+
+	TSharedPtr<FElectraHTTPStreamRequestApple, ESPMode::ThreadSafe> GetActiveTask(NSURLSessionTask* InTask);
+
+	static void RemoveOutdatedSessions()
+	{
+		FScopeLock lock(&AllSessionLock);
+		TArray<NSURLSession*> OutdatedSessions;
+		for(auto& Session : AllSessions)
+		{
+			if (!Session.Value.IsValid())
+			{
+				OutdatedSessions.Add(Session.Key);
+			}
+		}
+		for(auto& Outdated : OutdatedSessions)
+		{
+			AllSessions.Remove(Outdated);
+		}
+	}
+
+	// Handles
+	NSURLSession* Session = nil;
+
+	FThreadSafeCounter ExitRequest;
+	FRunnableThread* Thread = nullptr;
+	FTimeWaitableSignal HaveWorkSignal;
+
+	FElectraHTTPStreamThreadHandlerDelegate ThreadHandlerCallback;
+
+	FCriticalSection RequestLock;
+	TArray<TSharedPtr<FElectraHTTPStreamRequestApple, ESPMode::ThreadSafe>> NewRequests;
+	TArray<TSharedPtr<FElectraHTTPStreamRequestApple, ESPMode::ThreadSafe>> ActiveRequests;
+	TArray<TSharedPtr<FElectraHTTPStreamRequestApple, ESPMode::ThreadSafe>> CompletedRequests;
+
+	static FCriticalSection AllSessionLock;
+	static TMap<NSURLSession*, TWeakPtr<FElectraHTTPStreamApple, ESPMode::ThreadSafe>> AllSessions;
+};
+FCriticalSection FElectraHTTPStreamApple::AllSessionLock;
+TMap<NSURLSession*, TWeakPtr<FElectraHTTPStreamApple, ESPMode::ThreadSafe>> FElectraHTTPStreamApple::AllSessions;
+
+/***************************************************************************************************************************************************/
+/***************************************************************************************************************************************************/
+/***************************************************************************************************************************************************/
+
+/**
+ * Apple version of a HTTP stream request.
+ */
+class FElectraHTTPStreamRequestApple : public TSharedFromThis<FElectraHTTPStreamRequestApple, ESPMode::ThreadSafe>, public IElectraHTTPStreamRequest
+{
+public:
+	// This enum needs to be kept in order of data flow so states can be compared with < and >
+	enum EState
+	{
+		Inactive,
+		Connecting,
+		SendingRequest,
+		HeadersAvailable,
+		AwaitingResponseData,
+		ReceivingResponseData,
+		Completed,
+		Finished,
+		Error
+	};
+
+	enum EReadResponseResult
+	{
+		Success,
+		Failed,
+		EndOfData
+	};
+
+	FElectraHTTPStreamRequestApple();
+	virtual ~FElectraHTTPStreamRequestApple();
+
+	virtual void SetVerb(const FString& InVerb) override
+	{ Verb = InVerb; }
+	virtual IElectraHTTPStreamBuffer& POSTDataBuffer() override
+	{ return PostData; }
+	virtual void SetUserAgent(const FString& InUserAgent) override
+	{ UserAgent = InUserAgent; }
+	virtual void SetURL(const FString& InURL) override
+	{ URL = InURL; }
+	virtual void SetRange(const FString& InRange) override
+	{ Range = InRange; }
+	virtual void AllowCompression(bool bInAllowCompression) override
+	{ bAllowCompression = bInAllowCompression; }
+	virtual void AllowUnsafeRequestsForDebugging() override
+	{
+	#if ELECTRA_HTTPSTREAM_APPLE_ALLOW_UNSAFE_CONNECTIONS_FOR_DEBUGGING
+		bAllowUnsafeConnectionsForDebugging = true;
+	#endif
+	}
+	virtual void AddHeader(const FString& Header, const FString& Value, bool bAppendIfExists) override;
+	virtual FElectraHTTPStreamNotificationDelegate& NotificationDelegate() override
+	{ return NotificationCallback; }
+	virtual void Cancel() override;
+	virtual IElectraHTTPStreamResponsePtr GetResponse() override
+	{ return Response; }
+	virtual bool HasFailed() override
+	{ return Response->GetErrorMessage().Len() > 0; }
+	virtual FString GetErrorMessage() override
+	{ return Response->GetErrorMessage(); }
+
+	NSURLSessionTask* GetTaskHandle()
+	{ return TaskHandle; }
+
+	EState GetCurrentState()
+	{ return CurrentState; }
+	bool WasCanceled()
+	{ return bCancel; }
+	void DoCancel()
+	{
+		if (TaskHandle)
+		{
+			[TaskHandle cancel];
+		}
+	}
+	bool HasFinished()
+	{ return CurrentState == EState::Finished; }
+	bool Setup(TSharedPtr<FElectraHTTPStreamApple, ESPMode::ThreadSafe> InOwner);
+	bool Execute();
+	void Close();
+	void Terminate();
+	void SetError(NSError* InErrorCode);
+	void NotifyHeaders();
+	void NotifyDownloading();
+	void SetCompleted();
+	bool HasCompleted();
+	void SetFinished();
+	void SetResponseStatus(IElectraHTTPStreamResponse::EStatus InStatus);
+	void NotifyCallback(EElectraHTTPStreamNotificationReason InReason, int64 InParam);
+
+	void TriggerWorkSignal()
+	{
+		TSharedPtr<FElectraHTTPStreamApple, ESPMode::ThreadSafe> OwningManager = Owner.Pin();
+		if (OwningManager.IsValid())
+		{
+			OwningManager->TriggerWorkSignal();
+		}
+	}
+
+	NSURLSessionResponseDisposition ReceivedResponse(NSURLResponse* InResponse);
+	void ReceivedData(NSData* InData);
+	void ReceiveCompleted(NSError* InError);
+	NSURLRequest* ReceivedRedirect(NSHTTPURLResponse* InResponse, NSURLRequest* InRequest);
+	void SentData(int64 InNumBytesSent, int64 InNumTotalBytesSent, int64 InNumTotalBytesExpectedToSend);
+#if WITH_SSL
+	FElectraHTTPStreamApple::FChallengeResponse ReceivedChallenge(NSURLAuthenticationChallenge* InChallenge);
+#endif
+	void ReceivedMetrics(NSURLSessionTaskMetrics* InMetrics);
+
+protected:
+	FElectraHTTPStreamRequestApple(const FElectraHTTPStreamRequestApple&) = delete;
+	FElectraHTTPStreamRequestApple& operator=(const FElectraHTTPStreamRequestApple&) = delete;
+
+	bool ParseHeaders();
+
+	// User agent. Defaults to a global one but can be set with each individual request.
+	FString UserAgent;
+	// GET, HEAD or POST
+	FString Verb;
+	// URL to request
+	FString URL;
+	// Optional byte range. If set this must be a valid range string.
+	FString Range;
+	// Set to true to allow gzip/deflate for GET requests.
+	bool bAllowCompression = false;
+	// Additional headers to be sent with the request.
+	TMap<FString, FString> AdditionalHeaders;
+
+	// Owner
+	TWeakPtr<FElectraHTTPStreamApple, ESPMode::ThreadSafe> Owner;
+
+	// Configuration
+	int MaxRedirections = 10;
+#if ELECTRA_HTTPSTREAM_APPLE_ALLOW_UNSAFE_CONNECTIONS_FOR_DEBUGGING
+	bool bAllowUnsafeConnectionsForDebugging = false;
+#endif
+
+	FCriticalSection NotificationLock;
+	FElectraHTTPStreamNotificationDelegate NotificationCallback;
+
+	// Handles
+	NSURLSessionTask* TaskHandle = nil;
+
+	TArray<FElectraHTTPStreamHeader> CurrentResponseHeaders;
+	std::atomic_bool bResponseHeadersParsed { false };
+	FString EffectiveURL;
+	int32 HttpCode = 0;
+
+	std::atomic<EState> CurrentState { EState::Inactive };
+	std::atomic_bool bCancel { false };
+	std::atomic_bool bNotifiedHeaders { false };
+	int64 LastReportedDownloadSize = 0;
+
+	FElectraHTTPStreamBuffer PostData;
+	TSharedPtr<FElectraHTTPStreamResponse, ESPMode::ThreadSafe> Response;
+};
+
+
+FElectraHTTPStreamRequestApple::FElectraHTTPStreamRequestApple()
+{
+	UserAgent = ElectraHTTPStream::GetDefaultUserAgent();
+	Response = MakeShared<FElectraHTTPStreamResponse, ESPMode::ThreadSafe>();
+}
+
+FElectraHTTPStreamRequestApple::~FElectraHTTPStreamRequestApple()
+{
+	Close();
+}
+
+void FElectraHTTPStreamRequestApple::AddHeader(const FString& Header, const FString& Value, bool bAppendIfExists)
+{
+	// Ignore a few headers that will be set automatically.
+	if (Header.Equals(TEXT("User-Agent"), ESearchCase::IgnoreCase) ||
+		Header.Equals(TEXT("Accept-Encoding"), ESearchCase::IgnoreCase) ||
+		Header.Equals(TEXT("Transfer-Encoding"), ESearchCase::IgnoreCase) ||
+		Header.Equals(TEXT("Range"), ESearchCase::IgnoreCase) ||
+		Header.Equals(TEXT("Accept-Ranges"), ESearchCase::IgnoreCase))
+	{
+		return;
+	}
+	FString* ExistingHeader = AdditionalHeaders.Find(Header);
+	if (ExistingHeader)
+	{
+		// If the value is empty and we shall replace the header we need to remove it since empty headers are prohibited.
+		if (Value.IsEmpty() && !bAppendIfExists)
+		{
+			AdditionalHeaders.Remove(Header);
+		}
+		else if (Value.Len())
+		{
+			FString NewValue = *ExistingHeader;
+			NewValue.Append(TEXT(", "));
+			NewValue.Append(Value);
+			*ExistingHeader = NewValue;
+		}
+	}
+	// Header does not exist yet. If the value is empty we must not add the header since empty headers are prohibited.
+	else if (Value.Len())
+	{
+		AdditionalHeaders.Add(Header, Value);
+	}
+}
+
+void FElectraHTTPStreamRequestApple::Cancel()
+{
+	FScopeLock lock(&NotificationLock);
+	NotificationCallback.Unbind();
+	bCancel = true;
+}
+
+void FElectraHTTPStreamRequestApple::Terminate()
+{
+	Response->SetErrorMessage(TEXT("Terminated due to HTTP module shutdown"));
+	DoCancel();
+}
+
+void FElectraHTTPStreamRequestApple::SetError(NSError* InError)
+{
+	if (Response->GetErrorMessage().Len() == 0)
+	{
+		FString msg = ElectraHTTPStreamApple::GetErrorMessage(InError);
+		Response->SetErrorMessage(msg);
+		ElectraHTTPStreamApple::LogError(Response->GetErrorMessage());
+	}
+}
+
+void FElectraHTTPStreamRequestApple::NotifyHeaders()
+{
+	if (!bNotifiedHeaders)
+	{
+		bNotifiedHeaders = true;
+		// Parse the headers in case this was a HEAD request for which no response data is received.
+		ParseHeaders();
+		Response->CurrentState = IElectraHTTPStreamResponse::EState::ReceivedResponseHeaders;
+		NotifyCallback(EElectraHTTPStreamNotificationReason::ReceivedHeaders, Response->ResponseHeaders.Num());
+	}
+}
+
+void FElectraHTTPStreamRequestApple::NotifyDownloading()
+{
+	Response->CurrentState = IElectraHTTPStreamResponse::EState::ReceivingResponseData;
+	int64 nb = Response->GetResponseData().GetNumTotalBytesAdded();
+	if (nb > LastReportedDownloadSize)
+	{
+		NotifyCallback(EElectraHTTPStreamNotificationReason::ReadData, nb - LastReportedDownloadSize);
+		LastReportedDownloadSize = nb;
+	}
+}
+
+void FElectraHTTPStreamRequestApple::SetCompleted()
+{
+	CurrentState = EState::Completed;
+}
+
+bool FElectraHTTPStreamRequestApple::HasCompleted()
+{
+	return CurrentState == EState::Completed;
+}
+
+void FElectraHTTPStreamRequestApple::SetFinished()
+{
+	double Now = FPlatformTime::Seconds();
+	Response->TimeUntilFinished = Now - Response->StartTime;
+	Response->CurrentStatus = WasCanceled() ? IElectraHTTPStreamResponse::EStatus::Canceled : HasFailed() ?	IElectraHTTPStreamResponse::EStatus::Failed : IElectraHTTPStreamResponse::EStatus::Completed;
+	Response->CurrentState = IElectraHTTPStreamResponse::EState::Finished;
+	Response->SetEOS();
+	CurrentState = HasFailed() ? EState::Error : EState::Finished;
+}
+
+void FElectraHTTPStreamRequestApple::SetResponseStatus(IElectraHTTPStreamResponse::EStatus InStatus)
+{
+	Response->CurrentStatus = InStatus;
+}
+
+void FElectraHTTPStreamRequestApple::NotifyCallback(EElectraHTTPStreamNotificationReason InReason, int64 InParam)
+{
+	FScopeLock lock(&NotificationLock);
+	NotificationCallback.ExecuteIfBound(AsShared(), InReason, InParam);
+}
+
+bool FElectraHTTPStreamRequestApple::Setup(TSharedPtr<FElectraHTTPStreamApple, ESPMode::ThreadSafe> InOwner)
+{
+	Owner = InOwner;
+
+	// Check for a supported verb.
+	if (Verb.IsEmpty())
+	{
+		Verb = TEXT("GET");
+	}
+	if (!(Verb.Equals(TEXT("GET")) || Verb.Equals(TEXT("POST")) || Verb.Equals(TEXT("HEAD"))))
+	{
+		Response->SetErrorMessage(FString::Printf(TEXT("Unsupported verb \"%s\""), *Verb));
+		return false;
+	}
+
+	Electra::FURL_RFC3986 UrlParser;
+	if (!UrlParser.Parse(URL))
+	{
+		Response->SetErrorMessage(FString::Printf(TEXT("Failed to parse URL \"%s\""), *URL));
+		return false;
+	}
+
+	// GET/HEAD / POST specific options
+	if (Verb.Equals("GET") || Verb.Equals("HEAD"))
+	{
+		if (Range.Len())
+		{
+			FString RangeHdr = Range;
+			if (!RangeHdr.StartsWith(TEXT("bytes=")))
+			{
+				RangeHdr.InsertAt(0, TEXT("bytes="));
+			}
+			AdditionalHeaders.Add(TEXT("Range"), RangeHdr);
+		}
+
+		if (!bAllowCompression)
+		{
+			AdditionalHeaders.Add(TEXT("Accept-Encoding"), TEXT("identity"));
+		}
+	}
+
+	AdditionalHeaders.Add(TEXT("User-Agent"), UserAgent);
+
+	return true;
+}
+
+
+bool FElectraHTTPStreamRequestApple::Execute()
+{
+	SCOPED_AUTORELEASE_POOL;
+
+	TSharedPtr<FElectraHTTPStreamApple, ESPMode::ThreadSafe> OwningManager = Owner.Pin();
+	if (!OwningManager.IsValid())
+	{
+		return false;
+	}
+
+	// Set the origin URL as effective URL first in case there are problems or no redirections.
+	Response->EffectiveURL = URL;
+	Response->StartTime = FPlatformTime::Seconds();
+	Response->CurrentStatus = IElectraHTTPStreamResponse::EStatus::Running;
+
+	CurrentState = EState::Connecting;
+
+	NSMutableURLRequest* ur = [NSMutableURLRequest new];
+	ur.URL = [NSURL URLWithString: URL.GetNSString()];
+	ur.HTTPMethod = Verb.GetNSString();
+
+	// Add all additional headers.
+	for(auto &Hdr : AdditionalHeaders)
+	{
+		[ur setValue: Hdr.Value.GetNSString() forHTTPHeaderField: Hdr.Key.GetNSString()];
+	}
+
+	if (Verb.Equals("GET") || Verb.Equals("HEAD"))
+	{
+		TaskHandle = [OwningManager->GetSession() dataTaskWithRequest:ur];
+	}
+	else if (Verb.Equals("POST"))
+	{
+		// For now we need the EOS flag set as we send the data as a whole.
+		check(PostData.GetEOS());
+
+		const uint8* DataToSend;
+		int64 NumDataToSend;
+		PostData.LockBuffer(DataToSend, NumDataToSend);
+		PostData.UnlockBuffer(0);
+
+        TaskHandle = [OwningManager->GetSession() uploadTaskWithRequest:ur fromData:[NSData dataWithBytesNoCopy:(void*)DataToSend length:(NSUInteger)NumDataToSend freeWhenDone:false]];
+	}
+	if (TaskHandle)
+	{
+		// Retain the handle so we can cancel it later. If we do not do this we transfer ownership
+		// solely to the session which can drop it at any time.
+		[TaskHandle retain];
+		[TaskHandle resume];
+	}
+
+	return true;
+}
+
+bool FElectraHTTPStreamRequestApple::ParseHeaders()
+{
+	if (!bResponseHeadersParsed)
+	{
+		bResponseHeadersParsed = true;
+		for(auto &Hdr : CurrentResponseHeaders)
+		{
+			// Check for the most commonly used headers and set them on the side for faster access.
+			if (Hdr.Header.Equals(TEXT("Content-Type"), ESearchCase::IgnoreCase))
+			{
+				Response->ContentType = Hdr.Value;
+			}
+			else if (Hdr.Header.Equals(TEXT("Content-Length"), ESearchCase::IgnoreCase))
+			{
+				Response->ContentLength = Hdr.Value;
+			}
+			else if (Hdr.Header.Equals(TEXT("Content-Range"), ESearchCase::IgnoreCase))
+			{
+				Response->ContentRange = Hdr.Value;
+			}
+			else if (Hdr.Header.Equals(TEXT("Accept-Ranges"), ESearchCase::IgnoreCase))
+			{
+				Response->AcceptRanges = Hdr.Value;
+			}
+			else if (Hdr.Header.Equals(TEXT("Transfer-Encoding"), ESearchCase::IgnoreCase))
+			{
+				Response->TransferEncoding = Hdr.Value;
+			}
+
+			// Add to the list of headers, even those we have treated separately.
+			Response->ResponseHeaders.Emplace(MoveTemp(Hdr));
+		}
+
+		// Unfortunately we do not get access to the HTTP status line.
+		// The information whether or not HTTP/2 was used is only available at the end of the transfer
+		// when the availability of metrics is reported, which is too late for us.
+		// For lack of information we construct a HTTP/1.1 status with just the status value
+		Response->HTTPStatusLine = FString::Printf(TEXT("HTTP/1.1 %d"), HttpCode);
+		Response->HTTPResponseCode = HttpCode;
+		Response->EffectiveURL = EffectiveURL;
+
+		// Pre-allocate the receive buffer unless this is a HEAD request for which we will not get any data.
+		if (!Verb.Equals(TEXT("HEAD")))
+		{
+			// Check for Content-Length header
+			if (Response->ContentLength.Len())
+			{
+				int64 cl = -1;
+				LexFromString(cl, *Response->ContentLength);
+				if (cl >= 0)
+				{
+					Response->ResponseData.SetLengthFromResponseHeader(cl);
+				}
+			}
+			// Alternatively check Content-Range header
+			else if (Response->ContentRange.Len())
+			{
+				ElectraHTTPStream::FHttpRange ContentRange;
+				if (ContentRange.ParseFromContentRangeResponse(Response->ContentRange))
+				{
+					if (ContentRange.GetNumberOfBytes() >= 0)
+					{
+						Response->ResponseData.SetLengthFromResponseHeader(ContentRange.GetNumberOfBytes());
+					}
+				}
+			}
+		}
+	}
+	return true;
+}
+
+void FElectraHTTPStreamRequestApple::Close()
+{
+	if (TaskHandle)
+	{
+		[TaskHandle release];
+		TaskHandle = nil;
+    }
+
+	// Set EOS in the response receive buffer to signal no further data will arrive.
+	Response->SetEOS();
+}
+
+
+NSURLSessionResponseDisposition FElectraHTTPStreamRequestApple::ReceivedResponse(NSURLResponse* InResponse)
+{
+	SCOPED_AUTORELEASE_POOL;
+
+	LLM_SCOPE(ELLMTag::MediaStreaming);
+	SCOPE_CYCLE_COUNTER(STAT_ElectraHTTPThread_EventHandle);
+
+	// Get the response headers and status.
+	// Any URL load request's response is actually a NSHTTPURLResponse so it is safe to cast.
+	NSHTTPURLResponse* ur = (NSHTTPURLResponse*)InResponse;
+
+	NSDictionary* Headers = [ur allHeaderFields];
+
+	// Clear the (previous) response headers with every redirection.
+	CurrentResponseHeaders.Empty();
+	CurrentResponseHeaders.Reserve([Headers count]);
+	for (NSString* Key in [Headers allKeys])
+	{
+		FString ConvertedValue([Headers objectForKey:Key]);
+		FString ConvertedKey(Key);
+		FElectraHTTPStreamHeader hdr;
+		hdr.Header = ConvertedKey;
+		hdr.Value = ConvertedValue;
+		CurrentResponseHeaders.Emplace(MoveTemp(hdr));
+	}
+	NSInteger statusCode = [ur statusCode];
+	HttpCode = statusCode;
+	NSURL* Url = [ur URL];
+	FString effUrl([Url absoluteString]);
+	EffectiveURL = effUrl;
+
+	if (CurrentState < EState::HeadersAvailable)
+	{
+		double Now = FPlatformTime::Seconds();
+		double Elapsed = Now - Response->StartTime;
+		Response->TimeUntilNameResolved = Elapsed * 0.25;
+		Response->TimeUntilConnected = Elapsed * 0.5;
+		Response->TimeUntilHeadersAvailable = Elapsed;
+		Response->TimeOfMostRecentReceive = Now;
+		CurrentState = EState::HeadersAvailable;
+	}
+
+	TriggerWorkSignal();
+
+	return NSURLSessionResponseAllow;
+}
+
+void FElectraHTTPStreamRequestApple::ReceivedData(NSData* InData)
+{
+	LLM_SCOPE(ELLMTag::MediaStreaming);
+	SCOPE_CYCLE_COUNTER(STAT_ElectraHTTPThread_EventHandle);
+
+	double Now = FPlatformTime::Seconds();
+	Response->TimeOfMostRecentReceive = Now;
+	if (CurrentState < EState::ReceivingResponseData)
+	{
+		// Parse the headers, mostly to get to the content length to pre-allocate the receive buffer.
+		if (!ParseHeaders())
+		{
+			return;
+		}
+		Response->TimeUntilFirstByte = Now - Response->StartTime;
+		CurrentState = EState::ReceivingResponseData;
+	}
+
+	TConstArrayView<const uint8> Data(static_cast<const uint8*>([InData bytes]), (int32)[InData length]);
+	Response->AddResponseData(Data);
+
+	TriggerWorkSignal();
+}
+
+void FElectraHTTPStreamRequestApple::SentData(int64 InNumBytesSent, int64 InNumTotalBytesSent, int64 InNumTotalBytesExpectedToSend)
+{
+	LLM_SCOPE(ELLMTag::MediaStreaming);
+	SCOPE_CYCLE_COUNTER(STAT_ElectraHTTPThread_EventHandle);
+
+	double Now = FPlatformTime::Seconds();
+	Response->TimeOfMostRecentReceive = Now;
+	CurrentState = EState::SendingRequest;
+	if (InNumTotalBytesSent == PostData.GetNumTotalBytesAdded())
+	{
+		Response->TimeUntilRequestSent = Now - Response->StartTime;
+		Response->CurrentState = IElectraHTTPStreamResponse::EState::WaitingForResponseHeaders;
+	}
+	else
+	{
+		Response->CurrentState = IElectraHTTPStreamResponse::EState::SendingRequestData;
+	}
+}
+
+void FElectraHTTPStreamRequestApple::ReceiveCompleted(NSError* InError)
+{
+	LLM_SCOPE(ELLMTag::MediaStreaming);
+	SCOPE_CYCLE_COUNTER(STAT_ElectraHTTPThread_EventHandle);
+
+	if (InError)
+	{
+		SetError(InError);
+	}
+	SetCompleted();
+}
+
+NSURLRequest* FElectraHTTPStreamRequestApple::ReceivedRedirect(NSHTTPURLResponse* InResponse, NSURLRequest* InRequest)
+{
+	SCOPED_AUTORELEASE_POOL;
+
+	LLM_SCOPE(ELLMTag::MediaStreaming);
+	SCOPE_CYCLE_COUNTER(STAT_ElectraHTTPThread_EventHandle);
+
+	// Number of permitted redirections exhausted?
+	if (--MaxRedirections < 0)
+	{
+		return NULL;
+	}
+	// Otherwise, if this is a POST request we need to make sure it stays one
+	if (Verb.Equals(TEXT("POST")))
+	{
+		// If this is a 301 or 302 redirect for which the proposed new request wants to do a GET
+		// we create a new request from the proposed one and change it to a POST, for which we also need to
+		// put the data to send up again.
+		if ((InResponse.statusCode == 301 || InResponse.statusCode == 302) && [InRequest.HTTPMethod isEqualToString:@"GET"])
+		{
+			NSMutableURLRequest* ur = [InRequest mutableCopy];
+			ur.HTTPMethod = @"POST";
+			const uint8* DataToSend;
+			int64 NumDataToSend;
+			PostData.LockBuffer(DataToSend, NumDataToSend);
+			PostData.UnlockBuffer(0);
+			ur.HTTPBody = [NSData dataWithBytesNoCopy:(void*)DataToSend length:(NSUInteger)NumDataToSend freeWhenDone:false];
+			return ur;
+		}
+	}
+	// Just use the prepared new request.
+	return InRequest;
+}
+
+void FElectraHTTPStreamRequestApple::ReceivedMetrics(NSURLSessionTaskMetrics* InMetrics)
+{
+/*
+    for(NSURLSessionTaskTransactionMetrics* TaskMetric in InMetrics.transactionMetrics)
+    {
+        NSString* HTTPVersion = TaskMetric.networkProtocolName;
+    }
+*/
+}
+
+
+#if WITH_SSL
+FElectraHTTPStreamApple::FChallengeResponse FElectraHTTPStreamRequestApple::ReceivedChallenge(NSURLAuthenticationChallenge* InChallenge)
+{
+	FElectraHTTPStreamApple::FChallengeResponse	Resp;
+
+#if ELECTRA_HTTPSTREAM_APPLE_ALLOW_UNSAFE_CONNECTIONS_FOR_DEBUGGING
+	if (bAllowUnsafeConnectionsForDebugging)
+	{
+		if ([InChallenge.protectionSpace.authenticationMethod isEqualToString:NSURLAuthenticationMethodServerTrust])
+		{
+			// Here we could check for select hosts if we wanted to
+			//if ([challenge.protectionSpace.host isEqualToString:@"domaintoverride.com"])
+			{
+				Resp.Disposition = NSURLSessionAuthChallengeUseCredential;
+				Resp.Credential = [NSURLCredential credentialForTrust:InChallenge.protectionSpace.serverTrust];
+				return Resp;
+			}
+		}
+	}
+#endif
+
+	TSharedPtr<FElectraHTTPStreamApple, ESPMode::ThreadSafe> OwningManager = Owner.Pin();
+	if (OwningManager.IsValid())
+	{
+		Resp = OwningManager->ReceivedChallenge(InChallenge);
+	}
+	return Resp;
+}
+#endif
+
+
+/***************************************************************************************************************************************************/
+/***************************************************************************************************************************************************/
+/***************************************************************************************************************************************************/
+
+@interface FElectraHTTPStreamAppleSessionDelegate : NSObject<NSURLSessionDelegate, NSURLSessionTaskDelegate, NSURLSessionDataDelegate>
+// From NSURLSessionDelegate
+- (void)URLSession:(NSURLSession *)session didBecomeInvalidWithError:(NSError *)error;
+#if WITH_SSL && ELECTRA_HTTPSTREAM_APPLE_ENABLE_SESSION_CHALLENGE
+- (void)URLSession:(NSURLSession *)session didReceiveChallenge:(NSURLAuthenticationChallenge *)challenge completionHandler:(void (^)(NSURLSessionAuthChallengeDisposition disposition, NSURLCredential *credential))completionHandler;
+#endif
+
+// From NSURLSessionTaskDelegate
+- (void)URLSession:(NSURLSession *)session task:(NSURLSessionTask *)task didCompleteWithError:(NSError *)error;
+- (void)URLSession:(NSURLSession *)session task:(NSURLSessionTask *)task willPerformHTTPRedirection:(NSHTTPURLResponse *)response newRequest:(NSURLRequest *)request completionHandler:(void (^)(NSURLRequest *))completionHandler;
+- (void)URLSession:(NSURLSession *)session task:(NSURLSessionTask *)task didSendBodyData:(int64_t)bytesSent totalBytesSent:(int64_t)totalBytesSent totalBytesExpectedToSend:(int64_t)totalBytesExpectedToSend;
+- (void)URLSession:(NSURLSession *)session task:(NSURLSessionTask *)task needNewBodyStream:(void (^)(NSInputStream *bodyStream))completionHandler;
+#if WITH_SSL
+- (void)URLSession:(NSURLSession *)session task:(NSURLSessionTask *)task didReceiveChallenge:(NSURLAuthenticationChallenge *)challenge completionHandler:(void (^)(NSURLSessionAuthChallengeDisposition disposition, NSURLCredential *credential))completionHandler;
+#endif
+- (void)URLSession:(NSURLSession *)session taskIsWaitingForConnectivity:(NSURLSessionTask *)task;
+- (void)URLSession:(NSURLSession *)session task:(NSURLSessionTask *)task didFinishCollectingMetrics:(NSURLSessionTaskMetrics *)metrics;
+
+// From NSURLSessionDataDelegate
+- (void)URLSession:(NSURLSession *)session dataTask:(NSURLSessionDataTask *)dataTask didReceiveResponse:(NSURLResponse *)response completionHandler:(void (^)(NSURLSessionResponseDisposition disposition))completionHandler;
+- (void)URLSession:(NSURLSession *)session dataTask:(NSURLSessionDataTask *)dataTask didReceiveData:(NSData *)data;
+- (void)URLSession:(NSURLSession *)session dataTask:(NSURLSessionDataTask *)dataTask willCacheResponse:(NSCachedURLResponse *)proposedResponse completionHandler:(void (^)(NSCachedURLResponse *cachedResponse))completionHandler;
+@end
+
+@implementation FElectraHTTPStreamAppleSessionDelegate
+- (FElectraHTTPStreamAppleSessionDelegate*) init
+{
+	self = [super init];
+	return self;
+}
+- (void)dealloc
+{
+	[super dealloc];
+}
+- (void)URLSession:(NSURLSession *)session didBecomeInvalidWithError:(NSError *)error
+{
+	// This should normally be called only when we close the session and call [Session invalidateAndCancel];
+	//UE_LOG(LogElectraHTTPStream, Display, TEXT("========== NSURLSession became invalid"));
+}
+#if WITH_SSL && ELECTRA_HTTPSTREAM_APPLE_ENABLE_SESSION_CHALLENGE
+- (void)URLSession:(NSURLSession *)session didReceiveChallenge:(NSURLAuthenticationChallenge *)challenge completionHandler:(void (^)(NSURLSessionAuthChallengeDisposition disposition, NSURLCredential *credential))completionHandler
+{
+	// See https://developer.apple.com/library/archive/technotes/tn2232/_index.html
+	FElectraHTTPStreamApple::FChallengeResponse Resp;
+	TSharedPtr<FElectraHTTPStreamApple, ESPMode::ThreadSafe> Handler = FElectraHTTPStreamApple::GetHandlerForSession(session);
+	if (Handler.IsValid())
+	{
+		Resp = Handler->ReceivedChallenge(challenge);
+	}
+	completionHandler(Resp.Disposition, Resp.Credential);
+}
+#endif
+- (void)URLSession:(NSURLSession *)session task:(NSURLSessionTask *)task didCompleteWithError:(NSError *)error
+{
+	// Although this method is called "WithError" it is also without error at the end of the transfer!
+	TSharedPtr<FElectraHTTPStreamRequestApple, ESPMode::ThreadSafe> Request = FElectraHTTPStreamApple::GetTask(session, task);
+	if (Request.IsValid())
+	{
+		Request->ReceiveCompleted(error);
+	}
+}
+- (void)URLSession:(NSURLSession *)session task:(NSURLSessionTask *)task willPerformHTTPRedirection:(NSHTTPURLResponse *)response newRequest:(NSURLRequest *)request completionHandler:(void (^)(NSURLRequest *))completionHandler
+{
+	NSURLRequest* NewRequest = request;
+	TSharedPtr<FElectraHTTPStreamRequestApple, ESPMode::ThreadSafe> Request = FElectraHTTPStreamApple::GetTask(session, task);
+	if (Request.IsValid())
+	{
+		NewRequest = Request->ReceivedRedirect(response, request);
+	}
+	completionHandler(NewRequest);
+}
+- (void)URLSession:(NSURLSession *)session task:(NSURLSessionTask *)task didSendBodyData:(int64_t)bytesSent totalBytesSent:(int64_t)totalBytesSent totalBytesExpectedToSend:(int64_t)totalBytesExpectedToSend
+{
+	TSharedPtr<FElectraHTTPStreamRequestApple, ESPMode::ThreadSafe> Request = FElectraHTTPStreamApple::GetTask(session, task);
+	if (Request.IsValid())
+	{
+		Request->SentData(bytesSent, totalBytesSent, totalBytesExpectedToSend);
+	}
+}
+- (void)URLSession:(NSURLSession *)session task:(NSURLSessionTask *)task needNewBodyStream:(void (^)(NSInputStream *bodyStream))completionHandler
+{
+	UE_LOG(LogElectraHTTPStream, Warning, TEXT("NSURLSession::task::needNewBodyStream called, but upload stream is not used."));
+	completionHandler(NULL);
+}
+#if WITH_SSL
+- (void)URLSession:(NSURLSession *)session task:(NSURLSessionTask *)task didReceiveChallenge:(NSURLAuthenticationChallenge *)challenge completionHandler:(void (^)(NSURLSessionAuthChallengeDisposition disposition, NSURLCredential *credential))completionHandler
+{
+	// See https://developer.apple.com/library/archive/technotes/tn2232/_index.html
+	FElectraHTTPStreamApple::FChallengeResponse Resp;
+	TSharedPtr<FElectraHTTPStreamRequestApple, ESPMode::ThreadSafe> Request = FElectraHTTPStreamApple::GetTask(session, task);
+	if (Request.IsValid())
+	{
+		Resp = Request->ReceivedChallenge(challenge);
+	}
+	completionHandler(Resp.Disposition, Resp.Credential);
+}
+#endif
+- (void)URLSession:(NSURLSession *)session taskIsWaitingForConnectivity:(NSURLSessionTask *)task
+{
+	// Not using this for now.
+}
+- (void)URLSession:(NSURLSession *)session task:(NSURLSessionTask *)task didFinishCollectingMetrics:(NSURLSessionTaskMetrics *)metrics
+{
+	TSharedPtr<FElectraHTTPStreamRequestApple, ESPMode::ThreadSafe> Request = FElectraHTTPStreamApple::GetTask(session, task);
+	if (Request.IsValid())
+	{
+		Request->ReceivedMetrics(metrics);
+	}
+}
+- (void)URLSession:(NSURLSession *)session dataTask:(NSURLSessionDataTask *)dataTask didReceiveResponse:(NSURLResponse *)response completionHandler:(void (^)(NSURLSessionResponseDisposition disposition))completionHandler
+{
+	NSURLSessionResponseDisposition Action  = NSURLSessionResponseCancel;
+	TSharedPtr<FElectraHTTPStreamRequestApple, ESPMode::ThreadSafe> Request = FElectraHTTPStreamApple::GetTask(session, dataTask);
+	if (Request.IsValid())
+	{
+		Action = Request->ReceivedResponse(response);
+	}
+	completionHandler(Action);
+}
+- (void)URLSession:(NSURLSession *)session dataTask:(NSURLSessionDataTask *)dataTask didReceiveData:(NSData *)data
+{
+	TSharedPtr<FElectraHTTPStreamRequestApple, ESPMode::ThreadSafe> Request = FElectraHTTPStreamApple::GetTask(session, dataTask);
+	if (Request.IsValid())
+	{
+		Request->ReceivedData(data);
+	}
+}
+- (void)URLSession:(NSURLSession *)session dataTask:(NSURLSessionDataTask *)dataTask willCacheResponse:(NSCachedURLResponse *)proposedResponse completionHandler:(void (^)(NSCachedURLResponse *cachedResponse))completionHandler
+{
+	// Do not cache anything!
+	completionHandler(NULL);
+}
+@end
+
+
+
+
+
+FElectraHTTPStreamApple::FElectraHTTPStreamApple()
+{
+}
+
+FElectraHTTPStreamApple::~FElectraHTTPStreamApple()
+{
+	Close();
+}
+
+bool FElectraHTTPStreamApple::Initialize(const Electra::FParamDict& InOptions)
+{
+	SCOPED_AUTORELEASE_POOL;
+
+	LLM_SCOPE(ELLMTag::MediaStreaming);
+
+	RemoveOutdatedSessions();
+
+	NSURLSessionConfiguration* SessionConfig = [NSURLSessionConfiguration ephemeralSessionConfiguration];
+	SessionConfig.networkServiceType = NSURLNetworkServiceTypeResponsiveAV;
+	SessionConfig.allowsCellularAccess = YES;
+	SessionConfig.timeoutIntervalForRequest = 60.0;
+	SessionConfig.timeoutIntervalForResource = 180.0;
+	SessionConfig.waitsForConnectivity = NO;
+
+	SessionConfig.HTTPCookieAcceptPolicy = NSHTTPCookieAcceptPolicyNever;
+	SessionConfig.HTTPShouldSetCookies = NO;
+	SessionConfig.HTTPCookieStorage = nil;
+
+	//SessionConfig.TLSMinimumSupportedProtocolVersion = TLSv12;
+
+	//SessionConfig.HTTPMaximumConnectionsPerHost = 6;
+	//SessionConfig.HTTPShouldUsePipelining = YES;
+
+	// from iOS 13 on we could also:
+	//SessionConfig.allowsConstrainedNetworkAccess = YES;
+	//SessionConfig.allowsExpensiveNetworkAccess = YES;
+
+
+	if (InOptions.HaveKey(TEXT("proxy")))
+	{
+		FString ProxyAddressAndPort = InOptions.GetValue(TEXT("proxy")).SafeGetFString();
+		int32 ColonPos = INDEX_NONE;
+		if (ProxyAddressAndPort.FindLastChar(TCHAR(':'), ColonPos))
+		{
+			FString Host(ProxyAddressAndPort.Left(ColonPos));
+			FString PortStr(ProxyAddressAndPort.RightChop(ColonPos + 1));
+			int32 Port = 0;
+			LexFromString(Port, *PortStr);
+			if (Host.Len() && PortStr.Len() && Port)
+			{
+				// Create a proxy dictionary for the session configuration.
+				NSString* proxyHost = Host.GetNSString();
+				NSNumber* proxyPort = [NSNumber numberWithInt:Port];
+				NSDictionary *proxyDict = @{
+					(NSString *)kCFNetworkProxiesHTTPEnable : [NSNumber numberWithInt:1],
+					(NSString *)kCFNetworkProxiesHTTPProxy : proxyHost,
+					(NSString *)kCFNetworkProxiesHTTPPort : proxyPort,
+                // Sadly https proxy is not available on iOS / padOS / tvOS.
+                #if PLATFORM_MAC
+					(NSString *)kCFNetworkProxiesHTTPSEnable : [NSNumber numberWithInt:1],
+					(NSString *)kCFNetworkProxiesHTTPSProxy : proxyHost,
+					(NSString *)kCFNetworkProxiesHTTPSPort : proxyPort,
+                #endif
+                };
+				SessionConfig.connectionProxyDictionary = proxyDict;
+			}
+		}
+	}
+
+	// Create the delegate object. According to the documentation the session will keep a strong reference to it, so we
+	// do not need to. When we close the session the delegate object will be released as well.
+	FElectraHTTPStreamAppleSessionDelegate* SessionDelegate = [FElectraHTTPStreamAppleSessionDelegate new];
+    Session = [NSURLSession sessionWithConfiguration:SessionConfig delegate:(id<NSURLSessionDelegate>)SessionDelegate delegateQueue:nil];
+	if (!Session)
+	{
+		UE_LOG(LogElectraHTTPStream, Error, TEXT("Failed to create the NSURLSession"));
+		return false;
+	}
+
+    AllSessionLock.Lock();
+	AllSessions.Add(Session, AsShared());
+	AllSessionLock.Unlock();
+
+	// Create the worker thread.
+	Thread = FRunnableThread::Create(this, TEXT("ElectraHTTPStream"), 128 * 1024, TPri_Normal);
+	if (!Thread)
+	{
+		UE_LOG(LogElectraHTTPStream, Error, TEXT("Failed to create the ElectraHTTPStream worker thread"));
+		return false;
+	}
+	return true;
+}
+
+
+void FElectraHTTPStreamApple::Close()
+{
+	LLM_SCOPE(ELLMTag::MediaStreaming);
+
+	if (Thread)
+	{
+		Thread->Kill(true);
+		delete Thread;
+		Thread = nullptr;
+	}
+	if (Session)
+	{
+		AllSessionLock.Lock();
+		AllSessions.Remove(Session);
+		AllSessionLock.Unlock();
+
+		//[Session finishTasksAndInvalidate];
+		[Session invalidateAndCancel];
+		// Took out the release because it seems tht invalidateAndCancel does this internally
+		// as it does crash if we do it.
+		//[Session release];
+		Session = nil;
+	}
+	RemoveOutdatedSessions();
+}
+
+TSharedPtr<FElectraHTTPStreamRequestApple, ESPMode::ThreadSafe> FElectraHTTPStreamApple::GetActiveTask(NSURLSessionTask* InTask)
+{
+	FScopeLock lock(&RequestLock);
+	for(auto& Active : ActiveRequests)
+	{
+		if (Active->GetTaskHandle() == InTask)
+		{
+			return Active;
+		}
+	}
+	return nullptr;
+}
+
+IElectraHTTPStreamRequestPtr FElectraHTTPStreamApple::CreateRequest()
+{
+	LLM_SCOPE(ELLMTag::MediaStreaming);
+	return MakeShared<FElectraHTTPStreamRequestApple, ESPMode::ThreadSafe>();
+}
+
+void FElectraHTTPStreamApple::AddRequest(IElectraHTTPStreamRequestPtr InRequest)
+{
+	if (InRequest.IsValid())
+	{
+		if (Thread)
+		{
+			FScopeLock lock(&RequestLock);
+			NewRequests.Emplace(StaticCastSharedPtr<FElectraHTTPStreamRequestApple>(InRequest));
+			TriggerWorkSignal();
+		}
+		else
+		{
+			TSharedPtr<FElectraHTTPStreamRequestApple, ESPMode::ThreadSafe> Req = StaticCastSharedPtr<FElectraHTTPStreamRequestApple>(InRequest);
+			Req->NotifyCallback(EElectraHTTPStreamNotificationReason::Completed, 1);
+		}
+	}
+}
+
+void FElectraHTTPStreamApple::Stop()
+{
+	ExitRequest.Set(1);
+}
+
+uint32 FElectraHTTPStreamApple::Run()
+{
+	LLM_SCOPE(ELLMTag::MediaStreaming);
+
+	while(!ExitRequest.GetValue())
+	{
+		HaveWorkSignal.WaitTimeoutAndReset(20);
+
+		{
+		SCOPE_CYCLE_COUNTER(STAT_ElectraHTTPThread_Process);
+		SetupNewRequests();
+		UpdateActiveRequests();
+		HandleCompletedRequests();
+		}
+
+		// User callback
+		{
+		SCOPE_CYCLE_COUNTER(STAT_ElectraHTTPThread_CustomHandler);
+		ThreadHandlerCallback.ExecuteIfBound();
+		}
+	}
+
+	// Remove requests.
+	RequestLock.Lock();
+	while(NewRequests.Num())
+	{
+		TSharedPtr<FElectraHTTPStreamRequestApple, ESPMode::ThreadSafe> Req = NewRequests.Pop();
+		Req->Terminate();
+		CompletedRequests.Emplace(MoveTemp(Req));
+	}
+	while(ActiveRequests.Num())
+	{
+		TSharedPtr<FElectraHTTPStreamRequestApple, ESPMode::ThreadSafe> Req = ActiveRequests.Pop();
+		Req->Terminate();
+		CompletedRequests.Emplace(MoveTemp(Req));
+	}
+	RequestLock.Unlock();
+	HandleCompletedRequests();
+
+	return 0;
+}
+
+void FElectraHTTPStreamApple::SetupNewRequests()
+{
+	RequestLock.Lock();
+	TArray<TSharedPtr<FElectraHTTPStreamRequestApple, ESPMode::ThreadSafe>> NewReqs;
+	Swap(NewReqs, NewRequests);
+	RequestLock.Unlock();
+	for(auto &Request : NewReqs)
+	{
+		if (Request->Setup(AsShared()))
+		{
+			ActiveRequests.Emplace(Request);
+			if (Request->WasCanceled() || !Request->Execute())
+			{
+				ActiveRequests.Remove(Request);
+				CompletedRequests.Emplace(Request);
+			}
+		}
+		else
+		{
+			CompletedRequests.Emplace(Request);
+		}
+	}
+}
+
+void FElectraHTTPStreamApple::UpdateActiveRequests()
+{
+	for(int32 i=0; i<ActiveRequests.Num(); ++i)
+	{
+		TSharedPtr<FElectraHTTPStreamRequestApple, ESPMode::ThreadSafe> Request = ActiveRequests[i];
+		bool bRemoveRequest = Request->HasFinished() || Request->HasFailed() || Request->WasCanceled();
+
+		// If the request got canceled do not handle it any further and move it to the completed stage.
+		if (Request->WasCanceled())
+		{
+			Request->DoCancel();
+			Request->SetResponseStatus(IElectraHTTPStreamResponse::EStatus::Canceled);
+		}
+		else
+		{
+			// As we get notified about headers and received data the state of the request advances.
+			// To avoid invoking user callbacks from within the notification callback itself we check
+			// the request state against the response state here.
+			IElectraHTTPStreamResponsePtr Response = Request->GetResponse();
+
+			// Did we receive headers that we did not yet pass on?
+			if (Request->GetCurrentState() >= FElectraHTTPStreamRequestApple::EState::AwaitingResponseData && Response->GetState() < IElectraHTTPStreamResponse::EState::ReceivedResponseHeaders)
+			{
+				Request->NotifyHeaders();
+			}
+			// Report newly arrived data. Also notify headers if that did not happen already.
+			if (Request->GetCurrentState() >= FElectraHTTPStreamRequestApple::EState::ReceivingResponseData && Response->GetState() <= IElectraHTTPStreamResponse::EState::ReceivingResponseData)
+			{
+				Request->NotifyHeaders();
+				Request->NotifyDownloading();
+			}
+			// When completed, move to finished.
+			if (Request->HasCompleted())
+			{
+				Request->SetFinished();
+				bRemoveRequest = true;
+			}
+		}
+
+		if (bRemoveRequest)
+		{
+			ActiveRequests.RemoveAt(i);
+			--i;
+			CompletedRequests.Emplace(MoveTemp(Request));
+		}
+	}
+}
+
+void FElectraHTTPStreamApple::HandleCompletedRequests()
+{
+	if (CompletedRequests.Num())
+	{
+		TArray<TSharedPtr<FElectraHTTPStreamRequestApple, ESPMode::ThreadSafe>> TempRequests;
+		Swap(CompletedRequests, TempRequests);
+		for(auto &Finished : TempRequests)
+		{
+			Finished->Close();
+			// Parameter is 0 when finished successfully or 1 otherwise
+			int64 Param = Finished->HasFailed() ? 1 : 0;
+			if (!Finished->WasCanceled())
+			{
+				Finished->NotifyCallback(EElectraHTTPStreamNotificationReason::Completed, Param);
+			}
+		}
+	}
+}
+
+
+#if WITH_SSL
+FElectraHTTPStreamApple::FChallengeResponse FElectraHTTPStreamApple::ReceivedChallenge(NSURLAuthenticationChallenge* challenge)
+{
+	SCOPED_AUTORELEASE_POOL;
+
+	// CC gives the actual key, but strips the ASN.1 header... which means
+	// we can't calulate a proper SPKI hash without reconstructing it. sigh.
+	static const unsigned char rsa2048Asn1Header[] =
+	{
+		0x30, 0x82, 0x01, 0x22, 0x30, 0x0d, 0x06, 0x09,
+		0x2a, 0x86, 0x48, 0x86, 0xf7, 0x0d, 0x01, 0x01,
+		0x01, 0x05, 0x00, 0x03, 0x82, 0x01, 0x0f, 0x00
+	};
+	static const unsigned char rsa4096Asn1Header[] =
+	{
+		0x30, 0x82, 0x02, 0x22, 0x30, 0x0d, 0x06, 0x09,
+		0x2a, 0x86, 0x48, 0x86, 0xf7, 0x0d, 0x01, 0x01,
+		0x01, 0x05, 0x00, 0x03, 0x82, 0x02, 0x0f, 0x00
+	};
+	static const unsigned char ecdsaSecp256r1Asn1Header[] =
+	{
+		0x30, 0x59, 0x30, 0x13, 0x06, 0x07, 0x2a, 0x86,
+		0x48, 0xce, 0x3d, 0x02, 0x01, 0x06, 0x08, 0x2a,
+		0x86, 0x48, 0xce, 0x3d, 0x03, 0x01, 0x07, 0x03,
+		0x42, 0x00
+	};
+	static const unsigned char ecdsaSecp384r1Asn1Header[] =
+	{
+		0x30, 0x76, 0x30, 0x10, 0x06, 0x07, 0x2a, 0x86,
+		0x48, 0xce, 0x3d, 0x02, 0x01, 0x06, 0x05, 0x2b,
+		0x81, 0x04, 0x00, 0x22, 0x03, 0x62, 0x00
+	};
+
+	FChallengeResponse Response;
+	Response.Disposition = NSURLSessionAuthChallengeCancelAuthenticationChallenge;
+
+	if (ensure(ISslCertificateManager::PUBLIC_KEY_DIGEST_SIZE == CC_SHA256_DIGEST_LENGTH))
+	{
+		// we only care about challenges to the received certificate chain
+		if ([challenge.protectionSpace.authenticationMethod isEqualToString:NSURLAuthenticationMethodServerTrust])
+		{
+			SecTrustRef RemoteTrust = challenge.protectionSpace.serverTrust;
+			FString RemoteHost = FString(UTF8_TO_TCHAR([challenge.protectionSpace.host UTF8String]));
+			if ((RemoteTrust == NULL) || (RemoteHost.IsEmpty()))
+			{
+				UE_LOG(LogElectraHTTPStream, Error, TEXT("failed certificate pinning validation: could not parse parameters during certificate pinning evaluation"));
+				return Response;
+			}
+
+			if (!SecTrustEvaluateWithError(RemoteTrust, nil))
+			{
+				UE_LOG(LogElectraHTTPStream, Error, TEXT("failed certificate pinning validation: default certificate trust evaluation failed for domain '%s'"), *RemoteHost);
+				return Response;
+			}
+			// look at all certs in the remote chain and calculate the SHA256 hash of their DER-encoded SPKI
+			// the chain starts with the server's cert itself, so walk backwards to optimize for roots first
+			TArray<TArray<uint8, TFixedAllocator<ISslCertificateManager::PUBLIC_KEY_DIGEST_SIZE>>> CertDigests;
+
+			CFIndex NumCerts = SecTrustGetCertificateCount(RemoteTrust);
+			for(int i = static_cast<int>(NumCerts) - 1; i >= 0; i--)
+			{
+				SecCertificateRef Cert = SecTrustGetCertificateAtIndex(RemoteTrust, i);
+
+				// this is not great, but the only way to extract a public key from a SecCertificateRef
+				// is to create an individual SecTrustRef for each cert that only contains itself and then
+				// evaluate that against an empty X509 policy.
+				TCFRef<SecTrustRef> CertTrust;
+				TCFRef<SecPolicyRef> TrustPolicy = SecPolicyCreateBasicX509();
+				SecTrustCreateWithCertificates(Cert, TrustPolicy, CertTrust.GetForAssignment());
+				SecTrustEvaluateWithError(CertTrust, nil);
+				TCFRef<SecKeyRef> CertPubKey;
+
+				if (@available(macOS 11, *))
+				{
+					CertPubKey = SecTrustCopyKey(CertTrust);
+				}
+				else
+				{
+					PRAGMA_DISABLE_DEPRECATION_WARNINGS
+					// warning: 'SecTrustCopyPublicKey' is deprecated: first deprecated in iOS 14.0 [-Wdeprecated-declarations]
+
+					CertPubKey = SecTrustCopyPublicKey(CertTrust);
+
+					PRAGMA_ENABLE_DEPRECATION_WARNINGS
+				}
+
+				TCFRef<CFDataRef> CertPubKeyData = SecKeyCopyExternalRepresentation(CertPubKey, NULL);
+				if (!CertPubKeyData)
+				{
+					UE_LOG(LogElectraHTTPStream, Warning, TEXT("could not extract public key from certificate %i for domain '%s'; skipping!"), i, *RemoteHost);
+					continue;
+				}
+
+				// we got the key. now we have to figure out what type of key it is; thanks, CommonCrypto.
+				TCFRef<CFDictionaryRef> CertPubKeyAttr = SecKeyCopyAttributes(CertPubKey);
+				NSString *CertPubKeyType = static_cast<NSString *>(CFDictionaryGetValue(CertPubKeyAttr, kSecAttrKeyType));
+				NSNumber *CertPubKeySize = static_cast<NSNumber *>(CFDictionaryGetValue(CertPubKeyAttr, kSecAttrKeySizeInBits));
+				char *CertPubKeyASN1Header;
+				uint8_t CertPubKeyASN1HeaderSize = 0;
+				if ([CertPubKeyType isEqualToString: (NSString *)kSecAttrKeyTypeRSA])
+				{
+					switch ([CertPubKeySize integerValue])
+					{
+						case 2048:
+							UE_LOG(LogElectraHTTPStream, VeryVerbose, TEXT("found 2048 bit RSA pubkey"));
+							CertPubKeyASN1Header = (char *)rsa2048Asn1Header;
+							CertPubKeyASN1HeaderSize = sizeof(rsa2048Asn1Header);
+							break;
+						case 4096:
+							UE_LOG(LogElectraHTTPStream, VeryVerbose, TEXT("found 4096 bit RSA pubkey"));
+							CertPubKeyASN1Header = (char *)rsa4096Asn1Header;
+							CertPubKeyASN1HeaderSize = sizeof(rsa4096Asn1Header);
+							break;
+						default:
+							UE_LOG(LogElectraHTTPStream, Log, TEXT("unsupported RSA key length %i for certificate %i for domain '%s'; skipping!"), [CertPubKeySize integerValue], i, *RemoteHost);
+							continue;
+					}
+				}
+				else if ([CertPubKeyType isEqualToString: (NSString *)kSecAttrKeyTypeECSECPrimeRandom])
+				{
+					switch ([CertPubKeySize integerValue])
+					{
+						case 256:
+							UE_LOG(LogElectraHTTPStream, VeryVerbose, TEXT("found 256 bit ECDSA pubkey"));
+							CertPubKeyASN1Header = (char *)ecdsaSecp256r1Asn1Header;
+							CertPubKeyASN1HeaderSize = sizeof(ecdsaSecp256r1Asn1Header);
+							break;
+						case 384:
+							UE_LOG(LogElectraHTTPStream, VeryVerbose, TEXT("found 384 bit ECDSA pubkey"));
+							CertPubKeyASN1Header = (char *)ecdsaSecp384r1Asn1Header;
+							CertPubKeyASN1HeaderSize = sizeof(ecdsaSecp384r1Asn1Header);
+							break;
+						default:
+							UE_LOG(LogElectraHTTPStream, Log, TEXT("unsupported ECDSA key length %i for certificate %i for domain '%s'; skipping!"), [CertPubKeySize integerValue], i, *RemoteHost);
+							continue;
+					}
+				}
+				else
+				{
+					UE_LOG(LogElectraHTTPStream, Log, TEXT("unsupported key type (not RSA or ECDSA) for certificate %i for domain '%s'; skipping!"), i, *RemoteHost);
+					continue;
+				}
+
+				UE_LOG(LogElectraHTTPStream, VeryVerbose, TEXT("constructed key header: [%d] %s"), CertPubKeyASN1HeaderSize, UTF8_TO_TCHAR([[[NSData dataWithBytes:CertPubKeyASN1Header length:CertPubKeyASN1HeaderSize] description] UTF8String]));
+				UE_LOG(LogElectraHTTPStream, VeryVerbose, TEXT("current pubkey: [%d] %s"), [(NSData*)CertPubKeyData length], UTF8_TO_TCHAR([[[NSData dataWithBytes:[(NSData*)CertPubKeyData bytes] length:[(NSData*)CertPubKeyData length]] description] UTF8String]));
+
+				// smash 'em together to get a proper key with an ASN.1 header
+				NSMutableData *ReconstructedPubKey = [NSMutableData data];
+				[ReconstructedPubKey appendBytes:CertPubKeyASN1Header length:CertPubKeyASN1HeaderSize];
+				[ReconstructedPubKey appendData:CertPubKeyData];
+				UE_LOG(LogElectraHTTPStream, VeryVerbose, TEXT("reconstructed key: [%d] %s"), [ReconstructedPubKey length], UTF8_TO_TCHAR([[ReconstructedPubKey description] UTF8String]));
+
+				TArray<uint8, TFixedAllocator<ISslCertificateManager::PUBLIC_KEY_DIGEST_SIZE>> CertCalcDigest;
+				CertCalcDigest.AddUninitialized(CC_SHA256_DIGEST_LENGTH);
+				if (!CC_SHA256([ReconstructedPubKey bytes], (CC_LONG)[ReconstructedPubKey length], CertCalcDigest.GetData()))
+				{
+					UE_LOG(LogElectraHTTPStream, Warning, TEXT("could not calculate SHA256 digest of public key %d for domain '%s'; skipping!"), i, *RemoteHost);
+				}
+				else
+				{
+					CertDigests.Add(CertCalcDigest);
+					UE_LOG(LogElectraHTTPStream, Verbose, TEXT("added SHA256 digest to list for evaluation: domain: '%s' digest: [%d] %s"), *RemoteHost, CertCalcDigest.Num(), UTF8_TO_TCHAR([[[NSData dataWithBytes:CertCalcDigest.GetData() length:CertCalcDigest.Num()] description] UTF8String]));
+				}
+			}
+
+			//finally, see if any of the pubkeys in the chain match any of our pinned pubkey hashes
+			if (CertDigests.Num() <= 0 || !FSslModule::Get().GetCertificateManager().VerifySslCertificates(CertDigests, RemoteHost))
+			{
+				// we could not validate any of the provided certs in chain with the pinned hashes for this host
+				// so we tell the sender to cancel (which cancels the pending connection)
+				UE_LOG(LogElectraHTTPStream, Error, TEXT("failed certificate pinning validation: no SPKI hashes in request matched pinned hashes for domain '%s' (was provided %d certificates in request)"), *RemoteHost, CertDigests.Num());
+				return Response;
+			}
+		}
+	}
+	else
+	{
+		UE_LOG(LogElectraHTTPStream, Error, TEXT("failed certificate pinning validation: SslCertificateManager is using non-SHA256 SPKI hashes [expected %d bytes, got %d bytes]"), CC_SHA256_DIGEST_LENGTH, ISslCertificateManager::PUBLIC_KEY_DIGEST_SIZE);
+		return Response;
+	}
+
+	// if we got this far, pinning validation either succeeded or was disabled (or this was checking for client auth, etc.)
+	// so tell the connection to keep going with whatever else it was trying to validate
+	UE_LOG(LogElectraHTTPStream, Verbose, TEXT("certificate public key pinning either succeeded, is disabled, or challenge was not a server trust; continuing with auth"));
+	Response.Disposition = NSURLSessionAuthChallengePerformDefaultHandling;
+	return Response;
+}
+#endif
+
+
+/***************************************************************************************************************************************************/
+/***************************************************************************************************************************************************/
+/***************************************************************************************************************************************************/
+
+void FPlatformElectraHTTPStreamApple::Startup()
+{
+#if WITH_SSL
+	FSslModule& SslModule = FModuleManager::LoadModuleChecked<FSslModule>("SSL");
+	SslModule.GetSslManager().InitializeSsl();
+#endif
+}
+
+void FPlatformElectraHTTPStreamApple::Shutdown()
+{
+#if WITH_SSL
+	FSslModule& SslModule = FModuleManager::LoadModuleChecked<FSslModule>("SSL");
+	SslModule.GetSslManager().ShutdownSsl();
+#endif
+}
+
+
+TSharedPtr<IElectraHTTPStream, ESPMode::ThreadSafe> FPlatformElectraHTTPStreamApple::Create(const Electra::FParamDict& InOptions)
+{
+	TSharedPtr<FElectraHTTPStreamApple, ESPMode::ThreadSafe> New = MakeShareable(new FElectraHTTPStreamApple);
+	if (New.IsValid())
+	{
+		if (!New->Initialize(InOptions))
+		{
+			New.Reset();
+		}
+	}
+	return New;
+}
+
+
+#endif // ELECTRA_HTTPSTREAM_APPLE
diff --git a/Engine/Plugins/Media/ElectraUtil/Source/ElectraHTTPStream/Private/Apple/AppleElectraHTTPStream.h b/Engine/Plugins/Media/ElectraUtil/Source/ElectraHTTPStream/Private/Apple/AppleElectraHTTPStream.h
new file mode 100644
--- /dev/null
+++ b/Engine/Plugins/Media/ElectraUtil/Source/ElectraHTTPStream/Private/Apple/AppleElectraHTTPStream.h
@@ -0,0 +1,23 @@
+// Copyright Epic Games, Inc. All Rights Reserved.
+
+#pragma once
+
+#if ELECTRA_HTTPSTREAM_APPLE
+
+#include "CoreMinimal.h"
+#include "Templates/SharedPointer.h"
+#include "ParameterDictionary.h"
+
+class IElectraHTTPStream;
+
+class ELECTRAHTTPSTREAM_API FPlatformElectraHTTPStreamApple
+{
+public:
+	static void Startup();
+	static void Shutdown();
+	static TSharedPtr<IElectraHTTPStream, ESPMode::ThreadSafe> Create(const Electra::FParamDict& InOptions);
+};
+
+typedef FPlatformElectraHTTPStreamApple FPlatformElectraHTTPStream;
+
+#endif
diff --git a/Engine/Plugins/Media/ElectraUtil/Source/ElectraHTTPStream/Private/Apple/AppleLogError.cpp b/Engine/Plugins/Media/ElectraUtil/Source/ElectraHTTPStream/Private/Apple/AppleLogError.cpp
new file mode 100644
--- /dev/null
+++ b/Engine/Plugins/Media/ElectraUtil/Source/ElectraHTTPStream/Private/Apple/AppleLogError.cpp
@@ -0,0 +1,32 @@
+// Copyright Epic Games, Inc. All Rights Reserved.
+
+#if ELECTRA_HTTPSTREAM_APPLE
+
+#include "Apple/AppleLogError.h"
+#include "ElectraHTTPStreamModule.h"
+
+
+
+namespace ElectraHTTPStreamApple
+{
+
+FString GetErrorMessage(NSError* ErrorCode)
+{
+	SCOPED_AUTORELEASE_POOL;
+
+	FString Domain([ErrorCode domain]);
+	FString Descr([ErrorCode localizedDescription]);
+	int32 Code = ErrorCode.code;
+	FString msg = FString::Printf(TEXT("Error %d (0x%x) in %s: %s"), Code, Code, *Domain, *Descr);
+	return msg;
+}
+
+
+void LogError(const FString& Message)
+{
+	UE_LOG(LogElectraHTTPStream, Error, TEXT("%s"), *Message);
+}
+
+}
+
+#endif
diff --git a/Engine/Plugins/Media/ElectraUtil/Source/ElectraHTTPStream/Private/Apple/AppleLogError.h b/Engine/Plugins/Media/ElectraUtil/Source/ElectraHTTPStream/Private/Apple/AppleLogError.h
new file mode 100644
--- /dev/null
+++ b/Engine/Plugins/Media/ElectraUtil/Source/ElectraHTTPStream/Private/Apple/AppleLogError.h
@@ -0,0 +1,15 @@
+// Copyright Epic Games, Inc. All Rights Reserved.
+
+#pragma once
+
+#include "CoreMinimal.h"
+
+#if ELECTRA_HTTPSTREAM_APPLE
+
+namespace ElectraHTTPStreamApple
+{
+	FString GetErrorMessage(NSError* ErrorCode);
+	void LogError(const FString& Message);
+}
+
+#endif
diff --git a/Engine/Plugins/Media/ElectraUtil/Source/ElectraHTTPStream/Private/Curl/CurlElectra.h b/Engine/Plugins/Media/ElectraUtil/Source/ElectraHTTPStream/Private/Curl/CurlElectra.h
new file mode 100644
--- /dev/null
+++ b/Engine/Plugins/Media/ElectraUtil/Source/ElectraHTTPStream/Private/Curl/CurlElectra.h
@@ -0,0 +1,33 @@
+// Copyright Epic Games, Inc. All Rights Reserved.
+
+#pragma once
+
+#if ELECTRA_HTTPSTREAM_LIBCURL
+
+#include "CoreMinimal.h"
+
+#if PLATFORM_MICROSOFT
+#include "Microsoft/WindowsHWrapper.h"
+#include "Microsoft/AllowMicrosoftPlatformTypes.h"
+#endif
+
+#if ELECTRA_HTTPSTREAM_XCURL
+//We copied this template to include the windows file from WindowsHWrapper's way if including MinWindows.h, since including xcurl.h directly caused gnarly build errors
+#include "CoreTypes.h"
+#include "HAL/PlatformMemory.h"
+#include "Microsoft/PreWindowsApi.h"
+#ifndef STRICT
+#define STRICT
+#endif
+#include "XCurl.h"
+#include "Microsoft/PostWindowsApi.h"
+#else
+	#include "curl/curl.h"
+#endif
+
+#if PLATFORM_MICROSOFT
+#include "Microsoft/HideMicrosoftPlatformTypes.h"
+#endif
+
+
+#endif
