ElectraHTTPStream: Added a network module geared towards downloading of chunked media streams for use with the Electra media player.

diff --git a/Engine/Plugins/Media/ElectraUtil/Source/ElectraHTTPStream/Private/Curl/CurlElectraHTTPStream.cpp b/Engine/Plugins/Media/ElectraUtil/Source/ElectraHTTPStream/Private/Curl/CurlElectraHTTPStream.cpp
new file mode 100644
--- /dev/null
+++ b/Engine/Plugins/Media/ElectraUtil/Source/ElectraHTTPStream/Private/Curl/CurlElectraHTTPStream.cpp
@@ -0,0 +1,1307 @@
+// Copyright Epic Games, Inc. All Rights Reserved.
+
+#if ELECTRA_HTTPSTREAM_LIBCURL
+
+#include "Curl/CurlElectraHTTPStream.h"
+#include "HAL/ThreadSafeCounter.h"
+#include "HAL/PlatformTime.h"
+#include "HAL/PlatformProcess.h"
+#include "HAL/Runnable.h"
+#include "HAL/RunnableThread.h"
+#include "HAL/LowLevelMemTracker.h"
+
+#include "ElectraHTTPStreamModule.h"
+#include "ElectraHTTPStream.h"
+#include "ElectraHTTPStreamBuffer.h"
+#include "ElectraHTTPStreamResponse.h"
+#include "ElectraHTTPStreamPerf.h"
+
+#include "Utilities/TimeWaitableSignal.h"
+#include "Utilities/DefaultHttpUserAgent.h"
+#include "Utilities/HttpRangeHeader.h"
+#include "Utilities/URLParser.h"
+
+#include "Curl/CurlElectraHTTPStreamConfig.h"
+#include "Curl/CurlElectraHTTPStreamImpl.h"
+#include "Curl/CurlElectraHTTPStreamRequest.h"
+#include "Curl/CurlLogError.h"
+
+#if WITH_SSL
+#include "Ssl.h"
+#include <openssl/ssl.h>
+#endif
+
+#include <atomic>
+
+
+#define ELECTRA_HTTPSTREAM_INTERNAL_CURL_READ_BUFFER_SIZE	65536
+
+
+/***************************************************************************************************************************************************/
+/***************************************************************************************************************************************************/
+/***************************************************************************************************************************************************/
+
+DECLARE_STATS_GROUP(TEXT("Electra HTTP Stream"), STATGROUP_ElectraHTTPStream, STATCAT_Advanced);
+DECLARE_CYCLE_STAT(TEXT("Process"), STAT_ElectraHTTPThread_Process, STATGROUP_ElectraHTTPStream);
+DECLARE_CYCLE_STAT(TEXT("CurlMulti"), STAT_ElectraHTTPThread_CurlMulti, STATGROUP_ElectraHTTPStream);
+DECLARE_CYCLE_STAT(TEXT("Custom handler"), STAT_ElectraHTTPThread_CustomHandler, STATGROUP_ElectraHTTPStream);
+
+/***************************************************************************************************************************************************/
+/***************************************************************************************************************************************************/
+/***************************************************************************************************************************************************/
+
+FElectraHTTPStreamRequestLibCurl::FElectraHTTPStreamRequestLibCurl()
+{
+	UserAgent = ElectraHTTPStream::GetDefaultUserAgent();
+	Response = MakeShared<FElectraHTTPStreamResponse, ESPMode::ThreadSafe>();
+}
+
+FElectraHTTPStreamRequestLibCurl::~FElectraHTTPStreamRequestLibCurl()
+{
+	Close();
+}
+
+void FElectraHTTPStreamRequestLibCurl::SetVerb(const FString& InVerb)
+{ 
+	Verb = InVerb; 
+}
+
+IElectraHTTPStreamBuffer& FElectraHTTPStreamRequestLibCurl::POSTDataBuffer()
+{ 
+	return PostData; 
+}
+
+void FElectraHTTPStreamRequestLibCurl::SetUserAgent(const FString& InUserAgent)
+{ 
+	UserAgent = InUserAgent; 
+}
+
+void FElectraHTTPStreamRequestLibCurl::SetURL(const FString& InURL)
+{ 
+	URL = InURL; 
+}
+
+void FElectraHTTPStreamRequestLibCurl::SetRange(const FString& InRange)
+{ 
+	Range = InRange; 
+}
+
+void FElectraHTTPStreamRequestLibCurl::AllowCompression(bool bInAllowCompression)
+{ 
+	bAllowCompression = bInAllowCompression; 
+}
+
+void FElectraHTTPStreamRequestLibCurl::AllowUnsafeRequestsForDebugging()
+{
+#if ELECTRA_HTTPSTREAM_CURL_ALLOW_UNSAFE_CONNECTIONS_FOR_DEBUGGING
+	bAllowUnsafeConnectionsForDebugging = true;
+#endif
+}
+
+void FElectraHTTPStreamRequestLibCurl::AddHeader(const FString& Header, const FString& Value, bool bAppendIfExists)
+{
+	// Ignore a few headers that will be set automatically.
+	if (Header.Equals(TEXT("User-Agent"), ESearchCase::IgnoreCase) ||
+		Header.Equals(TEXT("Accept-Encoding"), ESearchCase::IgnoreCase) ||
+		Header.Equals(TEXT("Transfer-Encoding"), ESearchCase::IgnoreCase) ||
+		Header.Equals(TEXT("Range"), ESearchCase::IgnoreCase) ||
+		Header.Equals(TEXT("Accept-Ranges"), ESearchCase::IgnoreCase))
+	{
+		return;
+	}
+	FString* ExistingHeader = AdditionalHeaders.Find(Header);
+	if (ExistingHeader)
+	{
+		// If the value is empty and we shall replace the header we need to remove it since empty headers are prohibited.
+		if (Value.IsEmpty() && !bAppendIfExists)
+		{
+			AdditionalHeaders.Remove(Header);
+		}
+		else if (Value.Len())
+		{
+			FString NewValue = *ExistingHeader;
+			NewValue.Append(TEXT(", "));
+			NewValue.Append(Value);
+			*ExistingHeader = NewValue;
+		}
+	}
+	// Header does not exist yet. If the value is empty we must not add the header since empty headers are prohibited.
+	else if (Value.Len())
+	{
+		AdditionalHeaders.Add(Header, Value);
+	}
+}
+
+FElectraHTTPStreamNotificationDelegate& FElectraHTTPStreamRequestLibCurl::NotificationDelegate()
+{ 
+	return NotificationCallback; 
+}
+
+void FElectraHTTPStreamRequestLibCurl::Cancel()
+{
+	FScopeLock lock(&NotificationLock);
+	NotificationCallback.Unbind();
+	bCancel = true;
+}
+
+IElectraHTTPStreamResponsePtr FElectraHTTPStreamRequestLibCurl::GetResponse()
+{ 
+	return Response; 
+}
+
+bool FElectraHTTPStreamRequestLibCurl::HasFailed()
+{ 
+	return Response->GetErrorMessage().Len() > 0; 
+}
+
+FString FElectraHTTPStreamRequestLibCurl::GetErrorMessage()
+{ 
+	return Response->GetErrorMessage(); 
+}
+
+FElectraHTTPStreamRequestLibCurl::EState FElectraHTTPStreamRequestLibCurl::GetCurrentState()
+{ 
+	return CurrentState; 
+}
+
+bool FElectraHTTPStreamRequestLibCurl::WasCanceled()
+{ 
+	return bCancel; 
+}
+
+bool FElectraHTTPStreamRequestLibCurl::HasFinished()
+{ 
+	return CurrentState == EState::Finished; 
+}
+
+void FElectraHTTPStreamRequestLibCurl::Terminate()
+{
+	Response->SetErrorMessage(TEXT("Terminated due to HTTP module shutdown"));
+}
+
+void FElectraHTTPStreamRequestLibCurl::SetError(CURLcode InErrorCode)
+{
+	FString msg;
+	if (CurlErrorMessageBuffer[0])
+	{
+		msg = FString(CurlErrorMessageBuffer);
+	}
+	else
+	{
+		msg = ElectraHTTPStreamLibCurl::GetErrorMessage(InErrorCode);
+	}
+	if (Response->GetErrorMessage().Len() == 0)
+	{
+		Response->SetErrorMessage(msg);
+		ElectraHTTPStreamLibCurl::LogError(Response->GetErrorMessage());
+	}
+}
+
+void FElectraHTTPStreamRequestLibCurl::SetError(CURLMcode InErrorCode)
+{
+	FString msg;
+	if (CurlErrorMessageBuffer[0])
+	{
+		msg = FString(CurlErrorMessageBuffer);
+	}
+	else
+	{
+		msg = ElectraHTTPStreamLibCurl::GetErrorMessage(InErrorCode);
+	}
+	if (Response->GetErrorMessage().Len() == 0)
+	{
+		Response->SetErrorMessage(msg);
+		ElectraHTTPStreamLibCurl::LogError(Response->GetErrorMessage());
+	}
+}
+
+void FElectraHTTPStreamRequestLibCurl::NotifyHeaders()
+{
+	// Parse the headers in case this was a HEAD request for which no response data is received.
+	ParseHeaders();
+	Response->CurrentState = IElectraHTTPStreamResponse::EState::ReceivedResponseHeaders;
+	NotifyCallback(EElectraHTTPStreamNotificationReason::ReceivedHeaders, Response->ResponseHeaders.Num());
+}
+
+void FElectraHTTPStreamRequestLibCurl::NotifyDownloading()
+{
+	Response->CurrentState = IElectraHTTPStreamResponse::EState::ReceivingResponseData;
+	int64 nb = Response->GetResponseData().GetNumTotalBytesAdded();
+	if (nb > LastReportedDownloadSize)
+	{
+		NotifyCallback(EElectraHTTPStreamNotificationReason::ReadData, nb - LastReportedDownloadSize);
+		LastReportedDownloadSize = nb;
+	}
+}
+
+void FElectraHTTPStreamRequestLibCurl::SetCompleted()
+{
+	CurrentState = EState::Completed;
+}
+
+bool FElectraHTTPStreamRequestLibCurl::HasCompleted()
+{
+	return CurrentState == EState::Completed;
+}
+
+void FElectraHTTPStreamRequestLibCurl::SetFinished()
+{
+	double Now = FPlatformTime::Seconds();
+	Response->TimeUntilFinished = Now - Response->StartTime;
+	Response->CurrentStatus = WasCanceled() ? IElectraHTTPStreamResponse::EStatus::Canceled : HasFailed() ?	IElectraHTTPStreamResponse::EStatus::Failed : IElectraHTTPStreamResponse::EStatus::Completed;
+	Response->CurrentState = IElectraHTTPStreamResponse::EState::Finished;
+	Response->SetEOS();
+	CurrentState = HasFailed() ? EState::Error : EState::Finished;
+}
+
+void FElectraHTTPStreamRequestLibCurl::SetResponseStatus(IElectraHTTPStreamResponse::EStatus InStatus)
+{ 
+	Response->CurrentStatus = InStatus; 
+}
+
+void FElectraHTTPStreamRequestLibCurl::NotifyCallback(EElectraHTTPStreamNotificationReason InReason, int64 InParam)
+{
+	FScopeLock lock(&NotificationLock);
+	NotificationCallback.ExecuteIfBound(AsShared(), InReason, InParam);
+}
+
+#if WITH_SSL
+int FElectraHTTPStreamRequestLibCurl::SslCertVerify(int PreverifyOk, X509_STORE_CTX* Context)
+{
+	if (PreverifyOk == 1)
+	{
+		SSL* Handle = static_cast<SSL*>(X509_STORE_CTX_get_ex_data(Context, SSL_get_ex_data_X509_STORE_CTX_idx()));
+		if (!Handle)
+		{
+			return 0;
+		}
+		SSL_CTX* SslContext = SSL_get_SSL_CTX(Handle);
+		if (!SslContext)
+		{
+			return 0;
+		}
+		FElectraHTTPStreamRequestLibCurl* This = static_cast<FElectraHTTPStreamRequestLibCurl*>(SSL_CTX_get_app_data(SslContext));
+		if (!This)
+		{
+			return 0;
+		}
+		if (!FSslModule::Get().GetCertificateManager().VerifySslCertificates(Context, This->Host))
+		{
+			PreverifyOk = 0;
+		}
+	}
+	return PreverifyOk;
+}
+
+CURLcode FElectraHTTPStreamRequestLibCurl::CurlSSLCtxCallback(CURL* InEasyHandle, void* InSSLCtx)
+{
+	SSL_CTX* Context = static_cast<SSL_CTX*>(InSSLCtx);
+	const ISslCertificateManager& CertificateManager = FSslModule::Get().GetCertificateManager();
+
+	CertificateManager.AddCertificatesToSslContext(Context);
+
+	// For debugging purposes we may disable certificate validation for selected requests.
+	#if ELECTRA_HTTPSTREAM_CURL_ALLOW_UNSAFE_CONNECTIONS_FOR_DEBUGGING
+	if (!bAllowUnsafeConnectionsForDebugging)
+	#endif
+	{
+		SSL_CTX_set_verify(Context, SSL_CTX_get_verify_mode(Context), &FElectraHTTPStreamRequestLibCurl::SslCertVerify);
+		SSL_CTX_set_app_data(Context, this);
+	}
+	return CURLE_OK;
+}
+#endif
+
+bool FElectraHTTPStreamRequestLibCurl::Setup(FElectraHTTPStreamLibCurl* OwningManager)
+{
+	// Check for a supported verb.
+	if (Verb.IsEmpty())
+	{
+		Verb = TEXT("GET");
+	}
+	if (!(Verb.Equals(TEXT("GET")) || Verb.Equals(TEXT("POST")) || Verb.Equals(TEXT("HEAD"))))
+	{
+		Response->SetErrorMessage(FString::Printf(TEXT("Unsupported verb \"%s\""), *Verb));
+		ElectraHTTPStreamLibCurl::LogError(Response->GetErrorMessage());
+		return false;
+	}
+
+	Electra::FURL_RFC3986 UrlParser;
+	if (!UrlParser.Parse(URL))
+	{
+		Response->SetErrorMessage(FString::Printf(TEXT("Failed to parse URL \"%s\""), *URL));
+		ElectraHTTPStreamLibCurl::LogError(Response->GetErrorMessage());
+		return false;
+	}
+	Host = UrlParser.GetHost();
+
+	bool bIsSecure = UrlParser.GetScheme().Equals(TEXT("https"), ESearchCase::IgnoreCase);
+#if ELECTRA_HTTPSTREAM_REQUIRES_SECURE_CONNECTIONS
+	if (!bIsSecure)
+	{
+		#if ELECTRA_HTTPSTREAM_CURL_ALLOW_UNSAFE_CONNECTIONS_FOR_DEBUGGING
+		if (!bAllowUnsafeConnectionsForDebugging)
+		#endif
+		{
+			Response->SetErrorMessage(FString::Printf(TEXT("Attempted to create an insecure http:// request which is disabled on this platform")));
+			ElectraHTTPStreamLibCurl::LogError(Response->GetErrorMessage());
+			return false;
+		}
+	}
+#endif
+
+
+	// Create the easy handle.
+	CurlEasyHandle = curl_easy_init();
+	if (!CurlEasyHandle)
+	{
+		Response->SetErrorMessage(FString::Printf(TEXT("Calling curl_easy_init() failed")));
+		ElectraHTTPStreamLibCurl::LogError(Response->GetErrorMessage());
+		return false;
+	}
+
+	CURLcode Result;
+	Result = curl_easy_setopt(CurlEasyHandle, CURLOPT_ERRORBUFFER, CurlErrorMessageBuffer);
+
+	// Common options
+	Result = curl_easy_setopt(CurlEasyHandle, CURLOPT_PRIVATE, this);
+	Result = curl_easy_setopt(CurlEasyHandle, CURLOPT_NOSIGNAL, 1L);
+
+	Result = curl_easy_setopt(CurlEasyHandle, CURLOPT_HEADERDATA, this);
+	Result = curl_easy_setopt(CurlEasyHandle, CURLOPT_HEADERFUNCTION, &FElectraHTTPStreamRequestLibCurl::_CurlResponseHeaderCallback);
+	Result = curl_easy_setopt(CurlEasyHandle, CURLOPT_WRITEDATA, this);
+	Result = curl_easy_setopt(CurlEasyHandle, CURLOPT_WRITEFUNCTION, &FElectraHTTPStreamRequestLibCurl::_CurlResponseBodyCallback);
+#if ELECTRA_HTTPSTREAM_CURL_USE_XFERINFO
+	Result = curl_easy_setopt(CurlEasyHandle, CURLOPT_XFERINFODATA, this);
+	Result = curl_easy_setopt(CurlEasyHandle, CURLOPT_XFERINFOFUNCTION, &FElectraHTTPStreamRequestLibCurl::_CurlXferInfoCallback);
+#else
+	Result = curl_easy_setopt(CurlEasyHandle, CURLOPT_PROGRESSDATA, this);
+	Result = curl_easy_setopt(CurlEasyHandle, CURLOPT_PROGRESSFUNCTION, &FElectraHTTPStreamRequestLibCurl::_CurlProgressCallback);
+#endif
+	Result = curl_easy_setopt(CurlEasyHandle, CURLOPT_NOPROGRESS, 0L);
+
+	Result = curl_easy_setopt(CurlEasyHandle, CURLOPT_USERAGENT, TCHAR_TO_ANSI(*UserAgent));
+	Result = curl_easy_setopt(CurlEasyHandle, CURLOPT_URL, TCHAR_TO_ANSI(*URL));
+
+	Result = curl_easy_setopt(CurlEasyHandle, CURLOPT_BUFFERSIZE, (long)ELECTRA_HTTPSTREAM_INTERNAL_CURL_READ_BUFFER_SIZE);
+	Result = curl_easy_setopt(CurlEasyHandle, CURLOPT_FOLLOWLOCATION, 1L);
+	Result = curl_easy_setopt(CurlEasyHandle, CURLOPT_MAXREDIRS, (long)MaxRedirections);
+	if (ConnectionTimeoutMillis)
+	{
+		Result = curl_easy_setopt(CurlEasyHandle, CURLOPT_CONNECTTIMEOUT_MS, (long)ConnectionTimeoutMillis);
+	}
+
+	// HTTP/2 ?
+	if (OwningManager->SupportsHTTP2())
+	{
+#if defined(CURL_HTTP_VERSION_2) || LIBCURL_VERSION_NUM >= 0x072100
+		Result = curl_easy_setopt(CurlEasyHandle, CURLOPT_HTTP_VERSION, CURL_HTTP_VERSION_2_0);
+#endif
+	}
+
+#if ELECTRA_HTTPSTREAM_CURL_VERBOSE_DEBUGLOG
+	Result = curl_easy_setopt(CurlEasyHandle, CURLOPT_DEBUGDATA, this);
+	Result = curl_easy_setopt(CurlEasyHandle, CURLOPT_DEBUGFUNCTION, &FElectraHTTPStreamRequestLibCurl::_CurlDebugCallback);
+	Result = curl_easy_setopt(CurlEasyHandle, CURLOPT_VERBOSE, 1L);
+#endif
+
+#if ELECTRA_HTTPSTREAM_CURL_USE_SHARE
+	Result = curl_easy_setopt(CurlEasyHandle, CURLOPT_SHARE, OwningManager->GetCurlShareHandle());
+#endif
+
+#if ELECTRA_HTTPSTREAM_CURL_ALLOW_UNSAFE_CONNECTIONS_FOR_DEBUGGING
+	SetupSSL(bAllowUnsafeConnectionsForDebugging);
+#else
+	SetupSSL(false);
+#endif
+
+	SetupProxy(OwningManager->GetProxyAddressAndPort());
+
+	// GET/HEAD / POST specific options
+	if (Verb.Equals("GET") || Verb.Equals("HEAD"))
+	{
+		if (Range.Len())
+		{
+			FString RangeHdr = Range;
+			if (!RangeHdr.StartsWith(TEXT("bytes=")))
+			{
+				RangeHdr.InsertAt(0, TEXT("bytes="));
+			}
+			AdditionalHeaders.Add(TEXT("Range"), RangeHdr);
+		}
+
+		if (bAllowCompression)
+		{
+			Result = curl_easy_setopt(CurlEasyHandle, CURLOPT_ACCEPT_ENCODING, "");
+		}
+	}
+
+	// Add all additional headers.
+	for(auto &Hdr : AdditionalHeaders)
+	{
+		auto BuildHeader = [](const FString& InHdr, const FString& InVal) -> FString
+		{
+			if (InHdr.Len() && InVal.Len())
+			{
+				return InHdr + TEXT(": ") + InVal;
+			}
+			return FString();
+		};
+
+		FString Header = BuildHeader(Hdr.Key, Hdr.Value);
+		if (Header.Len())
+		{
+			CurlHeaderList = curl_slist_append(CurlHeaderList, TCHAR_TO_ANSI(*Header));
+		}
+	}
+	Result = curl_easy_setopt(CurlEasyHandle, CURLOPT_HTTPHEADER, CurlHeaderList);
+
+	return true;
+}
+
+void FElectraHTTPStreamRequestLibCurl::SetupSSL(bool bAllowUnsafeConnection)
+{
+	CURLcode Result = CURLE_OK; (void)Result;
+#if WITH_SSL
+	Result = curl_easy_setopt(CurlEasyHandle, CURLOPT_CAINFO, nullptr);
+	Result = curl_easy_setopt(CurlEasyHandle, CURLOPT_SSL_CTX_DATA, this);
+	Result = curl_easy_setopt(CurlEasyHandle, CURLOPT_SSL_CTX_FUNCTION, &FElectraHTTPStreamRequestLibCurl::_CurlSSLCtxCallback);
+#endif
+#if LIBCURL_VERSION_NUM >= 0x072200
+	Result = curl_easy_setopt(CurlEasyHandle, CURLOPT_SSLVERSION, CURL_SSLVERSION_TLSv1_1);
+#endif
+	if (bAllowUnsafeConnection)
+	{
+		Result = curl_easy_setopt(CurlEasyHandle, CURLOPT_SSL_VERIFYPEER, 0L);
+	}
+	else
+	{
+		Result = curl_easy_setopt(CurlEasyHandle, CURLOPT_SSL_VERIFYPEER, 1L);
+	}
+}
+
+void FElectraHTTPStreamRequestLibCurl::SetupProxy(const FString& InProxy)
+{
+	if (InProxy.Len())
+	{
+		CURLcode Result; (void)Result;
+		Result = curl_easy_setopt(CurlEasyHandle, CURLOPT_PROXY, TCHAR_TO_ANSI(*InProxy));
+	}
+}
+
+bool FElectraHTTPStreamRequestLibCurl::Execute(FElectraHTTPStreamLibCurl* OwningManager)
+{
+	// Set the origin URL as effective URL first in case there are problems or no redirections.
+	Response->EffectiveURL = URL;
+	Response->StartTime = FPlatformTime::Seconds();
+	Response->CurrentStatus = IElectraHTTPStreamResponse::EStatus::Running;
+
+	CurrentState = EState::Connecting;
+
+	CURLcode Result;
+	if (Verb.Equals("GET"))
+	{
+		Result = curl_easy_setopt(CurlEasyHandle, CURLOPT_HTTPGET, 1L);
+	}
+	else if (Verb.Equals("HEAD"))
+	{
+		Result = curl_easy_setopt(CurlEasyHandle, CURLOPT_NOBODY, 1L);
+	}
+	else if (Verb.Equals("POST"))
+	{
+		// For now we need the EOS flag set as we send the data as a whole.
+		check(PostData.GetEOS());
+
+		const uint8* DataToSend;
+		int64 NumDataToSend;
+		PostData.LockBuffer(DataToSend, NumDataToSend);
+		PostData.UnlockBuffer(0);
+
+		Result = curl_easy_setopt(CurlEasyHandle, CURLOPT_POST, 1L);
+		Result = curl_easy_setopt(CurlEasyHandle, CURLOPT_POSTFIELDS, NULL);
+#if ELECTRA_HTTPSTREAM_CURL_POST_VIA_INFILESIZE
+		Result = curl_easy_setopt(CurlEasyHandle, CURLOPT_INFILESIZE, (long)NumDataToSend);
+#else
+		Result = curl_easy_setopt(CurlEasyHandle, CURLOPT_POSTFIELDSIZE, (long)NumDataToSend);
+#endif
+		Result = curl_easy_setopt(CurlEasyHandle, CURLOPT_READDATA, this);
+		Result = curl_easy_setopt(CurlEasyHandle, CURLOPT_READFUNCTION, &FElectraHTTPStreamRequestLibCurl::_CurlUploadCallback);
+		Result = curl_easy_setopt(CurlEasyHandle, CURLOPT_SEEKDATA, this);
+		Result = curl_easy_setopt(CurlEasyHandle, CURLOPT_SEEKFUNCTION, &FElectraHTTPStreamRequestLibCurl::_CurlSeekCallback);
+		// Do *NOT* change POST to GET on any of the 301, 302 and 303 redirections!
+		Result = curl_easy_setopt(CurlEasyHandle, CURLOPT_POSTREDIR, CURL_REDIR_POST_ALL);
+	}
+
+	CurlMultiHandle = OwningManager->GetCurlMultiHandle();
+	CURLMcode MResult = curl_multi_add_handle(CurlMultiHandle, CurlEasyHandle);
+	if (MResult != CURLM_OK)
+	{
+		SetError(MResult);
+		return false;
+	}
+	return true;
+}
+
+#if ELECTRA_HTTPSTREAM_CURL_VERBOSE_DEBUGLOG
+int FElectraHTTPStreamRequestLibCurl::CurlDebugCallback(CURL* InEasyHandle, curl_infotype InInfoType, char* InData, size_t InLength)
+{
+	auto ToString = [](const char* InData, size_t InLength) -> FString
+	{
+		const char* NulChr = (const char*)memchr(InData, 0, InLength);
+		InLength = NulChr ? NulChr - InData : InLength;
+		auto ConvertedString = StringCast<TCHAR>(static_cast<const ANSICHAR*>(InData), (int32)InLength);
+		FString DebugText(ConvertedString.Length(), ConvertedString.Get());
+		DebugText.ReplaceInline(TEXT("\n"), TEXT(""), ESearchCase::CaseSensitive);
+		DebugText.ReplaceInline(TEXT("\r"), TEXT(""), ESearchCase::CaseSensitive);
+		return DebugText;
+	};
+	switch(InInfoType)
+	{
+		case CURLINFO_TEXT:
+		{
+			UE_LOG(LogElectraHTTPStream, Display, TEXT("%p: '%s'"), InEasyHandle, *ToString(InData, InLength));
+			break;
+		}
+		case CURLINFO_HEADER_IN:
+		{
+			UE_LOG(LogElectraHTTPStream, Display, TEXT("%p: '%s'"), InEasyHandle, *ToString(InData, InLength));
+			break;
+		}
+		default:
+		{
+			break;
+		}
+	}
+	return 0;
+}
+#endif
+
+#if ELECTRA_HTTPSTREAM_CURL_USE_XFERINFO
+int FElectraHTTPStreamRequestLibCurl::CurlXferInfoCallback(curl_off_t InDlTotal, curl_off_t InDlNow, curl_off_t InUlTotal, curl_off_t InUlNow)
+{
+#if ELECTRA_HTTPSTREAM_CURL_ALLOW_CANCEL_IN_CALLBACKS
+	return WasCanceled() ? 1 : 0;
+#else
+	return 0;
+#endif
+}
+#else
+int FElectraHTTPStreamRequestLibCurl::CurlProgressCallback(double InDlTotal, double InDlNow, double InUlTotal, double InUlNow)
+{
+	// Cancellation from within this callback was added with 7.23.00
+#if ELECTRA_HTTPSTREAM_CURL_ALLOW_CANCEL_IN_CALLBACKS && LIBCURL_VERSION_NUM >= 0x071700
+	return WasCanceled() ? 1 : 0;
+#else
+	return 0;
+#endif
+}
+#endif
+
+size_t FElectraHTTPStreamRequestLibCurl::CurlResponseHeaderCallback(void* InData, size_t InBlockSize, size_t InNumBlocks)
+{
+	// Received a response header.
+	// We get this for every header including those resulting in a redirect.
+	// As we are only interested in the headers from the final response we do not store
+	// the headers from a 300-303,305-399. 304 is not an actual redirection so we keep that.
+	if (CurrentState < EState::HeadersAvailable)
+	{
+		double Now = FPlatformTime::Seconds();
+		CurrentState = EState::HeadersAvailable;
+		Response->TimeUntilHeadersAvailable = Now - Response->StartTime;
+		Response->TimeOfMostRecentReceive = Now;
+
+		double StateTime = 0.0;
+		if (curl_easy_getinfo(CurlEasyHandle, CURLINFO_NAMELOOKUP_TIME, &StateTime) == CURLE_OK)
+		{
+			Response->TimeUntilNameResolved = StateTime;
+		}
+		if (curl_easy_getinfo(CurlEasyHandle, CURLINFO_CONNECT_TIME, &StateTime) == CURLE_OK)
+		{
+			Response->TimeUntilConnected = StateTime;
+		}
+	}
+	// Trailing headers?
+	else if (CurrentState >= EState::ReceivingResponseData)
+	{
+		// Not interested.
+		return InBlockSize * InNumBlocks;
+	}
+
+	size_t HeaderSize = InBlockSize * InNumBlocks;
+	if (HeaderSize > 0 && HeaderSize <= CURL_MAX_HTTP_HEADER)
+	{
+		FString HeaderLine((int32)HeaderSize, (const char*)InData);
+		int32 ColonPos = INDEX_NONE;
+		if (HeaderLine.FindChar(L':', ColonPos))
+		{
+			FString HeaderKey(HeaderLine.Left(ColonPos));
+			FString HeaderValue(HeaderLine.RightChop(ColonPos + 2));
+
+			HeaderKey.TrimStartAndEndInline();
+			HeaderValue.TrimStartAndEndInline();
+			FElectraHTTPStreamHeader Hdr;
+			Hdr.Header = HeaderKey;
+			Hdr.Value = HeaderValue;
+			CurrentResponseHeaders.Emplace(MoveTemp(Hdr));
+		}
+		else
+		{
+			// No colon, this should be the HTTP status line, eg. "HTTP/1.1 200 OK" or the empty header denoting the end.
+			HeaderLine.TrimStartAndEndInline();
+			if (HeaderLine.Len())
+			{
+				CurrentHTTPResponseStatusLine = HeaderLine;
+
+				// For simplicity we do not parse the status line for the code but as curl for it.
+				long HttpCodeNow = 0;
+				CURLcode Result = curl_easy_getinfo(CurlEasyHandle, CURLINFO_RESPONSE_CODE, &HttpCodeNow);
+				if (Result == CURLE_OK)
+				{
+					HttpCode = HttpCodeNow;
+					bIsRedirectionResponse = HttpCode >= 300 && HttpCode < 400 && HttpCode != 304;
+				}
+				// The status line is the first header, so if this indicates a redirection we can
+				// now clear the headers we have collected so far.
+				if (bIsRedirectionResponse)
+				{
+					CurrentResponseHeaders.Empty();
+				}
+
+				// Get the http version in use.
+				long HttpVersionUsed = CURL_HTTP_VERSION_NONE;
+				Result = curl_easy_getinfo(CurlEasyHandle, CURLINFO_HTTP_VERSION, &HttpVersionUsed);
+				if (Result == CURLE_OK)
+				{
+					HttpVersion = HttpVersionUsed;
+				}
+			}
+			else
+			{
+				// Get the effective URL
+				char* EffUrl = nullptr;
+				if (curl_easy_getinfo(CurlEasyHandle, CURLINFO_EFFECTIVE_URL, &EffUrl) == CURLE_OK)
+				{
+					EffectiveURL = FString(EffUrl);
+				}
+
+				// This is not necessarily the last header. There could be redirections following or
+				// there was a 100-continue or a 200-xxx proxy connection.
+			}
+		}
+	}
+	else
+	{
+		Response->SetErrorMessage(FString::Printf(TEXT("CurlResponseHeaderCallback() received bad header size of %u"), HeaderSize));
+		ElectraHTTPStreamLibCurl::LogError(Response->GetErrorMessage());
+		return 0;
+	}
+	return HeaderSize;
+}
+
+size_t FElectraHTTPStreamRequestLibCurl::CurlResponseBodyCallback(void* InData, size_t InBlockSize, size_t InNumBlocks)
+{
+	// Add the data to the response.
+	size_t nb = InBlockSize * InNumBlocks;
+
+	double Now = FPlatformTime::Seconds();
+	if (CurrentState < EState::ReceivingResponseData)
+	{
+		CurrentState = EState::ReceivingResponseData;
+		Response->TimeUntilFirstByte = Now - Response->StartTime;
+		// Parse the headers, mostly to get to the content length to pre-allocate the receive buffer.
+		if (!ParseHeaders())
+		{
+			return 0;
+		}
+	}
+	Response->TimeOfMostRecentReceive = Now;
+
+	TConstArrayView<const uint8> Data(static_cast<const uint8*>(InData), (int32)nb);
+	Response->AddResponseData(Data);
+
+#if ELECTRA_HTTPSTREAM_CURL_ALLOW_CANCEL_IN_CALLBACKS
+	return WasCanceled() ? 0 : nb;
+#else
+	return nb;
+#endif
+}
+
+size_t FElectraHTTPStreamRequestLibCurl::CurlUploadCallback(void* OutData, size_t InBlockSize, size_t InNumBlocks)
+{
+	int64 nb = (int64)(InBlockSize * InNumBlocks);
+
+	double Now = FPlatformTime::Seconds();
+	Response->TimeOfMostRecentReceive = Now;
+
+	const uint8* DataToSend;
+	int64 NumDataToSend;
+	PostData.LockBuffer(DataToSend, NumDataToSend);
+	int64 NumToCopy = NumDataToSend > nb ? nb : NumDataToSend;
+	FMemory::Memcpy(OutData, DataToSend, NumToCopy);
+	PostData.UnlockBuffer(NumToCopy);
+
+	if (PostData.GetNumTotalBytesHandedOut() == PostData.GetNumTotalBytesAdded())
+	{
+		Response->TimeUntilRequestSent = Now - Response->StartTime;
+	}
+
+#if ELECTRA_HTTPSTREAM_CURL_ALLOW_CANCEL_IN_CALLBACKS
+	return WasCanceled() ? CURL_READFUNC_ABORT : (size_t)NumToCopy;
+#else
+	return (size_t)NumToCopy;
+#endif
+}
+
+int FElectraHTTPStreamRequestLibCurl::CurlSeekCallback(curl_off_t InOffset, int InOrigin)
+{
+	// Only support re-sending the entire payload.
+	if (InOrigin == SEEK_SET && InOffset == 0)
+	{
+		const uint8* DataToSend;
+		int64 NumDataToSend;
+		PostData.LockBuffer(DataToSend, NumDataToSend);
+		bool bRewindOk = PostData.RewindToBeginning();
+		PostData.UnlockBuffer(0);
+		return bRewindOk ? CURL_SEEKFUNC_OK : CURL_SEEKFUNC_FAIL;
+	}
+	return CURL_SEEKFUNC_CANTSEEK;
+}
+
+bool FElectraHTTPStreamRequestLibCurl::ParseHeaders()
+{
+	if (!bResponseHeadersParsed)
+	{
+		for(auto &Hdr : CurrentResponseHeaders)
+		{
+			// Check for the most commonly used headers and set them on the side for faster access.
+			if (Hdr.Header.Equals(TEXT("Content-Type"), ESearchCase::IgnoreCase))
+			{
+				Response->ContentType = Hdr.Value;
+			}
+			else if (Hdr.Header.Equals(TEXT("Content-Length"), ESearchCase::IgnoreCase))
+			{
+				Response->ContentLength = Hdr.Value;
+			}
+			else if (Hdr.Header.Equals(TEXT("Content-Range"), ESearchCase::IgnoreCase))
+			{
+				Response->ContentRange = Hdr.Value;
+			}
+			else if (Hdr.Header.Equals(TEXT("Accept-Ranges"), ESearchCase::IgnoreCase))
+			{
+				Response->AcceptRanges = Hdr.Value;
+			}
+			else if (Hdr.Header.Equals(TEXT("Transfer-Encoding"), ESearchCase::IgnoreCase))
+			{
+				Response->TransferEncoding = Hdr.Value;
+			}
+
+			// Add to the list of headers, even those we have treated separately.
+			Response->ResponseHeaders.Emplace(MoveTemp(Hdr));
+		}
+		Response->HTTPStatusLine = CurrentHTTPResponseStatusLine;
+		Response->HTTPResponseCode = HttpCode;
+		Response->EffectiveURL = EffectiveURL;
+
+		// Pre-allocate the receive buffer unless this is a HEAD request for which we will not get any data.
+		if (!Verb.Equals(TEXT("HEAD")))
+		{
+			// Check for Content-Length header
+			if (Response->ContentLength.Len())
+			{
+				int64 cl = -1;
+				LexFromString(cl, *Response->ContentLength);
+				if (cl >= 0)
+				{
+					Response->ResponseData.SetLengthFromResponseHeader(cl);
+				}
+			}
+			// Alternatively check Content-Range header
+			else if (Response->ContentRange.Len())
+			{
+				ElectraHTTPStream::FHttpRange ContentRange;
+				if (ContentRange.ParseFromContentRangeResponse(Response->ContentRange))
+				{
+					if (ContentRange.GetNumberOfBytes() >= 0)
+					{
+						Response->ResponseData.SetLengthFromResponseHeader(ContentRange.GetNumberOfBytes());
+					}
+				}
+			}
+		}
+		bResponseHeadersParsed = true;
+	}
+	return true;
+}
+
+void FElectraHTTPStreamRequestLibCurl::Close()
+{
+	if (CurlEasyHandle)
+	{
+		if (CurlMultiHandle)
+		{
+			curl_multi_remove_handle(CurlMultiHandle, CurlEasyHandle);
+			CurlMultiHandle = nullptr;
+		}
+#if ELECTRA_HTTPSTREAM_CURL_USE_SHARE
+		curl_easy_setopt(CurlEasyHandle, CURLOPT_SHARE, nullptr);
+#endif
+		curl_easy_cleanup(CurlEasyHandle);
+		CurlEasyHandle = nullptr;
+	}
+
+	if (CurlHeaderList)
+	{
+		curl_slist_free_all(CurlHeaderList);
+		CurlHeaderList = nullptr;
+	}
+
+	// Set EOS in the response receive buffer to signal no further data will arrive.
+	Response->SetEOS();
+}
+
+/***************************************************************************************************************************************************/
+/***************************************************************************************************************************************************/
+/***************************************************************************************************************************************************/
+
+FElectraHTTPStreamLibCurl::FElectraHTTPStreamLibCurl()
+{
+}
+
+FElectraHTTPStreamLibCurl::~FElectraHTTPStreamLibCurl()
+{
+	Close();
+}
+
+bool FElectraHTTPStreamLibCurl::Initialize(const Electra::FParamDict& InOptions)
+{
+	LLM_SCOPE(ELLMTag::MediaStreaming);
+
+	if (InOptions.HaveKey(TEXT("proxy")))
+	{
+		ProxyAddressAndPort = InOptions.GetValue(TEXT("proxy")).SafeGetFString();
+	}
+
+	CURLSHcode scode = CURLSHE_OK; (void)scode;
+	CURLMcode mcode = CURLM_OK; (void)mcode;
+
+	const curl_version_info_data* CurlVersionInfo = curl_version_info(CURLVERSION_NOW);
+#ifdef CURL_VERSION_HTTP2
+	if ((CurlVersionInfo->features & CURL_VERSION_HTTP2) != 0)
+	{
+		bSupportsHTTP2 = true;
+	}
+#endif
+
+#if ELECTRA_HTTPSTREAM_CURL_USE_SHARE
+	// Create the share handle. Note that we do not need to set lock/unlock callbacks on it
+	// since we are not using this share from any thread other than our worker thread.
+	CurlShareHandle = curl_share_init();
+	if (!CurlShareHandle)
+	{
+		ElectraHTTPStreamLibCurl::LogError(TEXT("Calling curl_share_init() failed"));
+		return false;
+	}
+	#if LIBCURL_VERSION_NUM >= 0x071700
+	scode = curl_share_setopt(CurlShareHandle, CURLSHOPT_SHARE, CURL_LOCK_DATA_SSL_SESSION);
+	if (scode != CURLSHE_OK)
+	{
+		ElectraHTTPStreamLibCurl::LogError(FString::Printf(TEXT("Calling curl_share_setopt(CURL_LOCK_DATA_SSL_SESSION) failed with error %d"), scode));
+		return false;
+	}
+	#endif
+	#if LIBCURL_VERSION_NUM >= 0x073900
+	scode = curl_share_setopt(CurlShareHandle, CURLSHOPT_SHARE, CURL_LOCK_DATA_CONNECT);
+	if (scode != CURLSHE_OK)
+	{
+		ElectraHTTPStreamLibCurl::LogError(FString::Printf(TEXT("Calling curl_share_setopt(CURL_LOCK_DATA_CONNECT) failed with error %d"), scode));
+		return false;
+	}
+	#endif
+#endif
+
+	// Create the multi handle.
+	CurlMultiHandle = curl_multi_init();
+	if (!CurlMultiHandle)
+	{
+		ElectraHTTPStreamLibCurl::LogError(TEXT("Calling curl_multi_init() failed"));
+		return false;
+	}
+
+	if (bSupportsHTTP2)
+	{
+	#ifdef CURLPIPE_MULTIPLEX
+		// Try HTTP/2 multiplexing
+		mcode = curl_multi_setopt(CurlMultiHandle, CURLMOPT_PIPELINING, CURLPIPE_MULTIPLEX);
+	#endif
+	}
+
+	// Create the worker thread.
+	Thread = FRunnableThread::Create(this, TEXT("ElectraHTTPStream"), 128 * 1024, TPri_Normal);
+	if (!Thread)
+	{
+		UE_LOG(LogElectraHTTPStream, Error, TEXT("Failed to create the ElectraHTTPStream worker thread"));
+		return false;
+	}
+	return true;
+}
+
+
+void FElectraHTTPStreamLibCurl::Close()
+{
+	LLM_SCOPE(ELLMTag::MediaStreaming);
+
+	if (Thread)
+	{
+		Thread->Kill(true);
+		delete Thread;
+		Thread = nullptr;
+	}
+
+	if (CurlMultiHandle)
+	{
+		curl_multi_cleanup(CurlMultiHandle);
+		CurlMultiHandle = nullptr;
+	}
+
+#if ELECTRA_HTTPSTREAM_CURL_USE_SHARE
+	if (CurlShareHandle)
+	{
+		curl_share_cleanup(CurlShareHandle);
+		CurlShareHandle = nullptr;
+	}
+#endif
+}
+
+IElectraHTTPStreamRequestPtr FElectraHTTPStreamLibCurl::CreateRequest()
+{
+	LLM_SCOPE(ELLMTag::MediaStreaming);
+	return MakeShared<FElectraHTTPStreamRequestLibCurl, ESPMode::ThreadSafe>();
+}
+
+void FElectraHTTPStreamLibCurl::AddRequest(IElectraHTTPStreamRequestPtr InRequest)
+{
+	if (InRequest.IsValid())
+	{
+		if (Thread)
+		{
+			FScopeLock lock(&RequestLock);
+			NewRequests.Emplace(StaticCastSharedPtr<FElectraHTTPStreamRequestLibCurl>(InRequest));
+			TriggerWorkSignal();
+		}
+		else
+		{
+			TSharedPtr<FElectraHTTPStreamRequestLibCurl, ESPMode::ThreadSafe> Req = StaticCastSharedPtr<FElectraHTTPStreamRequestLibCurl>(InRequest);
+			Req->NotifyCallback(EElectraHTTPStreamNotificationReason::Completed, 1);
+		}
+	}
+}
+
+void FElectraHTTPStreamLibCurl::Stop()
+{
+	ExitRequest.Set(1);
+}
+
+
+void FElectraHTTPStreamLibCurl::WorkInnerLoop()
+{
+	double NextUpdateTime = 0.0;
+	while(!ExitRequest.GetValue())
+	{
+		/*
+			Historically Curl required the using application to perform a select() on its socket handles, which included
+			the ability to use timeouts. Unfortunately this requires knowledge about the type of handles used and the
+			inclusion of socket headers in the application.
+			Later curl_multi_wait() was added to alleviate this issue, but the timeout there is not working as expected.
+			If nothing is to do it will not wait at all and return immediately.
+			To fix that an even later version added curl_multi_poll(), which is what we would want to use, but cannot
+			because most of the versions of Curl we have available are too old and do not feature that method.
+
+			This leaves us with an ugly loop in which we need to sleep for small enough periods to not starve
+			curl_multi_perform().
+		*/
+
+		double tS = FPlatformTime::Seconds();
+		double tE = tS;
+		if (ActiveRequests.Num())
+		{
+			HandleCurl();
+			//int NumActiveHandles;
+			//CURLMcode Result = curl_multi_wait(CurlMultiHandle, nullptr, 0, 20, &NumActiveHandles);
+			double Elapsed = FPlatformTime::Seconds() - tS;
+			double MaxSleepPeriod = FMath::Max(MaxCurlMultiPerformDelay - Elapsed, 0.0);
+			FPlatformProcess::SleepNoStats(MaxSleepPeriod);
+		}
+		else
+		{
+			if (HaveWorkSignal.WaitTimeoutAndReset(20))
+			{
+				NextUpdateTime = 0.0;
+			}
+		}
+		tE = FPlatformTime::Seconds();
+
+		double TimeActive = tE - tS;
+		if ((NextUpdateTime -= TimeActive) <= 0.0)
+		{
+			NextUpdateTime += PeriodicRequestUpdateInterval;
+
+			{
+			SCOPE_CYCLE_COUNTER(STAT_ElectraHTTPThread_Process);
+			SetupNewRequests();
+			UpdateActiveRequests();
+			HandleCompletedRequests();
+			}
+
+			// User callback
+			{
+			SCOPE_CYCLE_COUNTER(STAT_ElectraHTTPThread_CustomHandler);
+			ThreadHandlerCallback.ExecuteIfBound();
+			}
+		}
+	}
+}
+
+void FElectraHTTPStreamLibCurl::WorkMultiPoll()
+{
+#if ELECTRA_HTTPSTREAM_CURL_USE_MULTIPOLL
+	double NextUpdateTime = 0.0;
+	while(!ExitRequest.GetValue())
+	{
+		int NumActiveHandles = 0;
+		double tS = FPlatformTime::Seconds();
+		CURLMcode Result = curl_multi_poll(CurlMultiHandle, nullptr, 0, 20, &NumActiveHandles);
+		double tE = FPlatformTime::Seconds();
+		double TimeActive = tE - tS;
+
+		HandleCurl();
+
+		if ((NextUpdateTime -= TimeActive) <= 0.0)
+		{
+			NextUpdateTime += PeriodicRequestUpdateInterval;
+
+			{
+			SCOPE_CYCLE_COUNTER(STAT_ElectraHTTPThread_Process);
+			SetupNewRequests();
+			UpdateActiveRequests();
+			HandleCompletedRequests();
+			}
+
+			// User callback
+			{
+			SCOPE_CYCLE_COUNTER(STAT_ElectraHTTPThread_CustomHandler);
+			ThreadHandlerCallback.ExecuteIfBound();
+			}
+		}
+	}
+#endif
+}
+
+
+uint32 FElectraHTTPStreamLibCurl::Run()
+{
+	LLM_SCOPE(ELLMTag::MediaStreaming);
+
+#if ELECTRA_HTTPSTREAM_CURL_USE_MULTIPOLL
+	WorkMultiPoll();
+#else
+	WorkInnerLoop();
+#endif
+
+	// Remove requests.
+	RequestLock.Lock();
+	while(NewRequests.Num())
+	{
+		TSharedPtr<FElectraHTTPStreamRequestLibCurl, ESPMode::ThreadSafe> Req = NewRequests.Pop();
+		Req->Terminate();
+		CompletedRequests.Emplace(MoveTemp(Req));
+	}
+	while(ActiveRequests.Num())
+	{
+		TSharedPtr<FElectraHTTPStreamRequestLibCurl, ESPMode::ThreadSafe> Req = ActiveRequests.Pop();
+		Req->Terminate();
+		CompletedRequests.Emplace(MoveTemp(Req));
+	}
+	RequestLock.Unlock();
+	HandleCompletedRequests();
+
+	return 0;
+}
+
+void FElectraHTTPStreamLibCurl::SetupNewRequests()
+{
+	RequestLock.Lock();
+	TArray<TSharedPtr<FElectraHTTPStreamRequestLibCurl, ESPMode::ThreadSafe>> NewReqs;
+	Swap(NewReqs, NewRequests);
+	RequestLock.Unlock();
+	for(auto &Request : NewReqs)
+	{
+		if (Request->Setup(this))
+		{
+			ActiveRequests.Emplace(Request);
+			if (Request->WasCanceled() || !Request->Execute(this))
+			{
+				ActiveRequests.Remove(Request);
+				CompletedRequests.Emplace(Request);
+			}
+		}
+		else
+		{
+			Request->SetFinished();
+			CompletedRequests.Emplace(Request);
+		}
+	}
+}
+
+
+void FElectraHTTPStreamLibCurl::HandleCurl()
+{
+	SCOPE_CYCLE_COUNTER(STAT_ElectraHTTPThread_CurlMulti);
+
+	int RunningHandles = -1;
+	CURLMcode mc;
+	do 
+	{
+		mc = curl_multi_perform(CurlMultiHandle, &RunningHandles);
+	} while(mc == CURLM_CALL_MULTI_PERFORM);
+	CURLMsg* CurlMessage = nullptr;
+	int NumMessagesLeft = 0;
+	while(1)
+	{
+		CurlMessage = curl_multi_info_read(CurlMultiHandle, &NumMessagesLeft);
+		if (!CurlMessage)
+		{
+			break;
+		}
+		if (CurlMessage->msg != CURLMSG_DONE)
+		{
+			continue;
+		}
+
+		CURL* EasyHandle = CurlMessage->easy_handle;
+		FElectraHTTPStreamRequestLibCurl* Request = nullptr;
+		CURLcode Err = curl_easy_getinfo(EasyHandle, CURLINFO_PRIVATE, &Request);
+		if (CurlMessage->data.result != CURLE_OK)
+		{
+			if (!Request->WasCanceled())
+			{
+				Request->SetError(CurlMessage->data.result);
+			}
+		}
+		// Tag the request as having completed, but not finished yet!
+		Request->SetCompleted();
+	}
+}
+
+void FElectraHTTPStreamLibCurl::UpdateActiveRequests()
+{
+	for(int32 i=0; i<ActiveRequests.Num(); ++i)
+	{
+		TSharedPtr<FElectraHTTPStreamRequestLibCurl, ESPMode::ThreadSafe> Request = ActiveRequests[i];
+		bool bRemoveRequest = Request->HasFinished() || Request->HasFailed() || Request->WasCanceled();
+
+		// If the request got canceled do not handle it any further and move it to the completed stage.
+		if (Request->WasCanceled())
+		{
+			Request->SetResponseStatus(IElectraHTTPStreamResponse::EStatus::Canceled);
+		}
+		else
+		{
+			// As curl notifies us about headers and received data the state of the request advances.
+			// To avoid invoking user callbacks from within curl itself we check the request state against
+			// the response state here.
+			IElectraHTTPStreamResponsePtr Response = Request->GetResponse();
+
+			// Did we receive headers from Curl we did not yet pass on?
+			if (Request->GetCurrentState() >= FElectraHTTPStreamRequestLibCurl::EState::AwaitingResponseData && Response->GetState() < IElectraHTTPStreamResponse::EState::ReceivedResponseHeaders)
+			{
+				Request->NotifyHeaders();
+			}
+			// Report newly arrived data
+			if (Request->GetCurrentState() >= FElectraHTTPStreamRequestLibCurl::EState::ReceivingResponseData && Response->GetState() <= IElectraHTTPStreamResponse::EState::ReceivingResponseData)
+			{
+				Request->NotifyDownloading();
+			}
+			// When completed, move to finished.
+			if (Request->HasCompleted())
+			{
+				Request->SetFinished();
+				bRemoveRequest = true;
+			}
+		}
+
+		if (bRemoveRequest)
+		{
+			ActiveRequests.RemoveAt(i);
+			--i;
+			CompletedRequests.Emplace(MoveTemp(Request));
+		}
+	}
+}
+
+void FElectraHTTPStreamLibCurl::HandleCompletedRequests()
+{
+	if (CompletedRequests.Num())
+	{
+		TArray<TSharedPtr<FElectraHTTPStreamRequestLibCurl, ESPMode::ThreadSafe>> TempRequests;
+		Swap(CompletedRequests, TempRequests);
+		for(auto &Finished : TempRequests)
+		{
+			Finished->Close();
+			// Parameter is 0 when finished successfully or 1 otherwise
+			int64 Param = Finished->HasFailed() ? 1 : 0;
+			if (!Finished->WasCanceled())
+			{
+				Finished->NotifyCallback(EElectraHTTPStreamNotificationReason::Completed, Param);
+			}
+		}
+	}
+}
+
+
+/***************************************************************************************************************************************************/
+/***************************************************************************************************************************************************/
+/***************************************************************************************************************************************************/
+
+#include "HttpModule.h"
+
+void FPlatformElectraHTTPStreamLibCurl::Startup()
+{
+	// Because of global methods like curl_global_init_mem() and similar methods for SSL have to rely on
+	// the engine's HTTP module doing this. We are sort of piggybacking onto this since there is no way for
+	// us to do the initialization over again.
+	// We also do not know if the engine HTTP module is even using Curl on this platform. We _try_ to keep
+	// our use of Curl in sync with what the engine is using.
+
+	FModuleManager::LoadModuleChecked<FHttpModule>("HTTP");
+
+	// Do the same for SSL. Here we get a hold of it and need to close it down in Shutdown().
+#if WITH_SSL
+	FSslModule& SslModule = FModuleManager::LoadModuleChecked<FSslModule>("SSL");
+	SslModule.GetSslManager().InitializeSsl();
+#endif
+}
+
+void FPlatformElectraHTTPStreamLibCurl::Shutdown()
+{
+#if WITH_SSL
+	FSslModule& SslModule = FModuleManager::LoadModuleChecked<FSslModule>("SSL");
+	SslModule.GetSslManager().ShutdownSsl();
+#endif
+}
+
+
+TSharedPtr<IElectraHTTPStream, ESPMode::ThreadSafe> FPlatformElectraHTTPStreamLibCurl::Create(const Electra::FParamDict& InOptions)
+{
+	TSharedPtr<FElectraHTTPStreamLibCurl, ESPMode::ThreadSafe> New = MakeShareable(new FElectraHTTPStreamLibCurl);
+	if (New.IsValid())
+	{
+		if (!New->Initialize(InOptions))
+		{
+			New.Reset();
+		}
+	}
+	return New;
+}
+
+
+#endif // ELECTRA_HTTPSTREAM_LIBCURL
diff --git a/Engine/Plugins/Media/ElectraUtil/Source/ElectraHTTPStream/Private/Curl/CurlElectraHTTPStream.h b/Engine/Plugins/Media/ElectraUtil/Source/ElectraHTTPStream/Private/Curl/CurlElectraHTTPStream.h
new file mode 100644
--- /dev/null
+++ b/Engine/Plugins/Media/ElectraUtil/Source/ElectraHTTPStream/Private/Curl/CurlElectraHTTPStream.h
@@ -0,0 +1,31 @@
+// Copyright Epic Games, Inc. All Rights Reserved.
+
+#pragma once
+
+#if ELECTRA_HTTPSTREAM_LIBCURL
+
+#include "CoreMinimal.h"
+#include "Templates/SharedPointer.h"
+#include "ParameterDictionary.h"
+
+class IElectraHTTPStream;
+
+class ELECTRAHTTPSTREAM_API FPlatformElectraHTTPStreamLibCurl
+{
+public:
+	static void Startup();
+	static void Shutdown();
+	static TSharedPtr<IElectraHTTPStream, ESPMode::ThreadSafe> Create(const Electra::FParamDict& InOptions);
+};
+
+
+#ifndef ELECTRA_HTTPSTREAM_CURL_PLATFORM_OVERRIDE
+#define ELECTRA_HTTPSTREAM_CURL_PLATFORM_OVERRIDE 0
+#endif
+
+// Does the platform override the generic libCurl implementation?
+#if ELECTRA_HTTPSTREAM_CURL_PLATFORM_OVERRIDE == 0
+typedef FPlatformElectraHTTPStreamLibCurl FPlatformElectraHTTPStream;
+#endif
+
+#endif
diff --git a/Engine/Plugins/Media/ElectraUtil/Source/ElectraHTTPStream/Private/Curl/CurlElectraHTTPStreamConfig.h b/Engine/Plugins/Media/ElectraUtil/Source/ElectraHTTPStream/Private/Curl/CurlElectraHTTPStreamConfig.h
new file mode 100644
--- /dev/null
+++ b/Engine/Plugins/Media/ElectraUtil/Source/ElectraHTTPStream/Private/Curl/CurlElectraHTTPStreamConfig.h
@@ -0,0 +1,52 @@
+// Copyright Epic Games, Inc. All Rights Reserved.
+
+#pragma once
+
+#if ELECTRA_HTTPSTREAM_LIBCURL
+
+#include "Curl/CurlElectra.h"
+
+#if ELECTRA_HTTPSTREAM_XCURL == 0
+#define ELECTRA_HTTPSTREAM_CURL_USE_SHARE 1
+#define ELECTRA_HTTPSTREAM_CURL_POST_VIA_INFILESIZE 0
+#else
+#define ELECTRA_HTTPSTREAM_CURL_USE_SHARE 0
+#define ELECTRA_HTTPSTREAM_CURL_POST_VIA_INFILESIZE 1
+#define ELECTRA_HTTPSTREAM_CURL_USE_XFERINFO 1
+#define ELECTRA_HTTPSTREAM_CURL_USE_MULTIPOLL 1
+#endif
+
+#ifndef LIBCURL_VERSION_NUM
+#define LIBCURL_VERSION_NUM 0
+#endif
+
+#ifndef ELECTRA_HTTPSTREAM_CURL_USE_XFERINFO
+#if LIBCURL_VERSION_NUM >= 0x072200
+#define ELECTRA_HTTPSTREAM_CURL_USE_XFERINFO 1
+#else
+#define ELECTRA_HTTPSTREAM_CURL_USE_XFERINFO 0
+#endif
+#endif
+
+#ifndef ELECTRA_HTTPSTREAM_CURL_USE_MULTIPOLL
+#if LIBCURL_VERSION_NUM >= 0x074200
+#define ELECTRA_HTTPSTREAM_CURL_USE_MULTIPOLL 1
+#else
+#define ELECTRA_HTTPSTREAM_CURL_USE_MULTIPOLL 0
+#endif
+#endif
+
+// Set this to 1 to allow aborts in Curl callbacks. This does not seem to work with xCurl though.
+#define ELECTRA_HTTPSTREAM_CURL_ALLOW_CANCEL_IN_CALLBACKS 0
+
+// In all but shipping builds we allow disabling of security checks.
+#if !UE_BUILD_SHIPPING
+#define ELECTRA_HTTPSTREAM_CURL_ALLOW_UNSAFE_CONNECTIONS_FOR_DEBUGGING 1
+#else
+#define ELECTRA_HTTPSTREAM_CURL_ALLOW_UNSAFE_CONNECTIONS_FOR_DEBUGGING 0
+#endif
+
+#define ELECTRA_HTTPSTREAM_CURL_VERBOSE_DEBUGLOG 0
+
+
+#endif // ELECTRA_HTTPSTREAM_LIBCURL
diff --git a/Engine/Plugins/Media/ElectraUtil/Source/ElectraHTTPStream/Private/Curl/CurlElectraHTTPStreamImpl.h b/Engine/Plugins/Media/ElectraUtil/Source/ElectraHTTPStream/Private/Curl/CurlElectraHTTPStreamImpl.h
new file mode 100644
--- /dev/null
+++ b/Engine/Plugins/Media/ElectraUtil/Source/ElectraHTTPStream/Private/Curl/CurlElectraHTTPStreamImpl.h
@@ -0,0 +1,89 @@
+// Copyright Epic Games, Inc. All Rights Reserved.
+
+#pragma once
+
+#if ELECTRA_HTTPSTREAM_LIBCURL
+
+#include "Curl/CurlElectraHTTPStreamConfig.h"
+#include "Curl/CurlElectra.h"
+#include "Utilities/TimeWaitableSignal.h"
+
+class FElectraHTTPStreamRequestLibCurl;
+
+/**
+ * libCurl version of the ElectraHTTPStream implementation.
+ */
+class FElectraHTTPStreamLibCurl : public TSharedFromThis<FElectraHTTPStreamLibCurl, ESPMode::ThreadSafe>, public IElectraHTTPStream, private FRunnable
+{
+public:
+	virtual ~FElectraHTTPStreamLibCurl();
+
+	FElectraHTTPStreamLibCurl();
+	virtual bool Initialize(const Electra::FParamDict& InOptions);
+
+	virtual FElectraHTTPStreamThreadHandlerDelegate& ThreadHandlerDelegate() override
+	{ return ThreadHandlerCallback; }
+
+	virtual void Close() override;
+
+	virtual IElectraHTTPStreamRequestPtr CreateRequest() override;
+
+	virtual void AddRequest(IElectraHTTPStreamRequestPtr Request) override;
+
+#if ELECTRA_HTTPSTREAM_CURL_USE_SHARE
+	CURLSH* GetCurlShareHandle()
+	{ return CurlShareHandle; }
+#endif
+	CURLM* GetCurlMultiHandle()
+	{ return CurlMultiHandle; }
+	bool SupportsHTTP2()
+	{ return bSupportsHTTP2; }
+	FString GetProxyAddressAndPort()
+	{ return ProxyAddressAndPort; }
+
+	void TriggerWorkSignal()
+	{ 
+		HaveWorkSignal.Signal(); 
+#if ELECTRA_HTTPSTREAM_CURL_USE_MULTIPOLL
+		curl_multi_wakeup(CurlMultiHandle);
+#endif
+	}
+
+private:
+	// Methods from FRunnable
+	virtual uint32 Run() override final;
+	virtual void Stop() override final;
+
+	void SetupNewRequests();
+	void UpdateActiveRequests();
+	void HandleCurl();
+	void HandleCompletedRequests();
+	void WorkInnerLoop();
+	void WorkMultiPoll();
+
+	// Configuration
+	FString ProxyAddressAndPort;
+	double PeriodicRequestUpdateInterval = 0.05;	// 20ms
+	double MaxCurlMultiPerformDelay = 0.001;		// 1ms
+
+
+	// Curl handles
+#if ELECTRA_HTTPSTREAM_CURL_USE_SHARE
+	CURLSH* CurlShareHandle = nullptr;
+#endif
+	CURLM* CurlMultiHandle = nullptr;
+	bool bSupportsHTTP2 = false;
+
+	FThreadSafeCounter ExitRequest;
+	FRunnableThread* Thread = nullptr;
+	FTimeWaitableSignal HaveWorkSignal;
+
+	FElectraHTTPStreamThreadHandlerDelegate ThreadHandlerCallback;
+
+	FCriticalSection RequestLock;
+	TArray<TSharedPtr<FElectraHTTPStreamRequestLibCurl, ESPMode::ThreadSafe>> NewRequests;
+	TArray<TSharedPtr<FElectraHTTPStreamRequestLibCurl, ESPMode::ThreadSafe>> ActiveRequests;
+	TArray<TSharedPtr<FElectraHTTPStreamRequestLibCurl, ESPMode::ThreadSafe>> CompletedRequests;
+};
+
+#endif
diff --git a/Engine/Plugins/Media/ElectraUtil/Source/ElectraHTTPStream/Private/Curl/CurlElectraHTTPStreamRequest.h b/Engine/Plugins/Media/ElectraUtil/Source/ElectraHTTPStream/Private/Curl/CurlElectraHTTPStreamRequest.h
new file mode 100644
--- /dev/null
+++ b/Engine/Plugins/Media/ElectraUtil/Source/ElectraHTTPStream/Private/Curl/CurlElectraHTTPStreamRequest.h
@@ -0,0 +1,170 @@
+// Copyright Epic Games, Inc. All Rights Reserved.
+
+#pragma once
+
+#if ELECTRA_HTTPSTREAM_LIBCURL
+
+#include "Curl/CurlElectraHTTPStreamConfig.h"
+#include "Curl/CurlElectra.h"
+#include "ElectraHTTPStreamResponse.h"
+
+#if WITH_SSL
+#include "Ssl.h"
+#include <openssl/ssl.h>
+#endif
+
+class FElectraHTTPStreamLibCurl;
+class FElectraHTTPStreamResponse;
+
+
+/**
+ * libCurl version of a HTTP stream request.
+ */
+class FElectraHTTPStreamRequestLibCurl : public TSharedFromThis<FElectraHTTPStreamRequestLibCurl, ESPMode::ThreadSafe>, public IElectraHTTPStreamRequest
+{
+public:
+	// This enum needs to be kept in order of data flow so states can be compared with < and >
+	enum EState
+	{
+		Inactive,
+		Connecting,
+		SendingRequest,
+		HeadersAvailable,
+		AwaitingResponseData,
+		ReceivingResponseData,
+		Completed,
+		Finished,
+		Error
+	};
+
+	enum EReadResponseResult
+	{
+		Success,
+		Failed,
+		EndOfData
+	};
+
+	FElectraHTTPStreamRequestLibCurl();
+	virtual ~FElectraHTTPStreamRequestLibCurl();
+
+	virtual void SetVerb(const FString& InVerb) override;
+	virtual IElectraHTTPStreamBuffer& POSTDataBuffer() override;
+	virtual void SetUserAgent(const FString& InUserAgent) override;
+	virtual void SetURL(const FString& InURL) override;
+	virtual void SetRange(const FString& InRange) override;
+	virtual void AllowCompression(bool bInAllowCompression) override;
+	virtual void AllowUnsafeRequestsForDebugging() override;
+	virtual void AddHeader(const FString& Header, const FString& Value, bool bAppendIfExists) override;
+	virtual FElectraHTTPStreamNotificationDelegate& NotificationDelegate() override;
+	virtual void Cancel() override;
+	virtual IElectraHTTPStreamResponsePtr GetResponse() override;
+	virtual bool HasFailed() override;
+	virtual FString GetErrorMessage() override;
+
+	EState GetCurrentState();
+	bool WasCanceled();
+	bool HasFinished();
+	bool Setup(FElectraHTTPStreamLibCurl* OwningManager);
+	virtual void SetupSSL(bool bAllowUnsafeConnection);
+	virtual void SetupProxy(const FString& InProxy);
+	bool Execute(FElectraHTTPStreamLibCurl* OwningManager);
+	void Close();
+	void Terminate();
+	void SetError(CURLcode InErrorCode);
+	void SetError(CURLMcode InErrorCode);
+	void NotifyHeaders();
+	void NotifyDownloading();
+	void SetCompleted();
+	bool HasCompleted();
+	void SetFinished();
+	void SetResponseStatus(IElectraHTTPStreamResponse::EStatus InStatus);
+	void NotifyCallback(EElectraHTTPStreamNotificationReason InReason, int64 InParam);
+
+protected:
+	FElectraHTTPStreamRequestLibCurl(const FElectraHTTPStreamRequestLibCurl&) = delete;
+	FElectraHTTPStreamRequestLibCurl& operator=(const FElectraHTTPStreamRequestLibCurl&) = delete;
+
+	bool ParseHeaders();
+
+#if ELECTRA_HTTPSTREAM_CURL_VERBOSE_DEBUGLOG
+	virtual int CurlDebugCallback(CURL* InEasyHandle, curl_infotype InInfoType, char* InData, size_t InLength);
+	static int _CurlDebugCallback(CURL* InEasyHandle, curl_infotype InInfoType, char* InData, size_t InLength, void* InUserData)
+	{ return static_cast<FElectraHTTPStreamRequestLibCurl*>(InUserData)->CurlDebugCallback(InEasyHandle, InInfoType, InData, InLength); }
+#endif
+#if ELECTRA_HTTPSTREAM_CURL_USE_XFERINFO
+	virtual int CurlXferInfoCallback(curl_off_t InDlTotal, curl_off_t InDlNow, curl_off_t InUlTotal, curl_off_t InUlNow);
+	static int _CurlXferInfoCallback(void* InUserData, curl_off_t InDlTotal, curl_off_t InDlNow, curl_off_t InUlTotal, curl_off_t InUlNow)
+	{ return static_cast<FElectraHTTPStreamRequestLibCurl*>(InUserData)->CurlXferInfoCallback(InDlTotal, InDlNow, InUlTotal, InUlNow); }
+#else
+	virtual int CurlProgressCallback(double InDlTotal, double InDlNow, double InUlTotal, double InUlNow);
+	static int _CurlProgressCallback(void* InUserData, double InDlTotal, double InDlNow, double InUlTotal, double InUlNow)
+	{ return static_cast<FElectraHTTPStreamRequestLibCurl*>(InUserData)->CurlProgressCallback(InDlTotal, InDlNow, InUlTotal, InUlNow); }
+#endif
+	virtual size_t CurlResponseHeaderCallback(void* InData, size_t InBlockSize, size_t InNumBlocks);
+	static size_t _CurlResponseHeaderCallback(void* InData, size_t InBlockSize, size_t InNumBlocks, void* InUserData)
+	{ return static_cast<FElectraHTTPStreamRequestLibCurl*>(InUserData)->CurlResponseHeaderCallback(InData, InBlockSize, InNumBlocks); }
+	virtual size_t CurlResponseBodyCallback(void* InData, size_t InBlockSize, size_t InNumBlocks);
+	static size_t _CurlResponseBodyCallback(void* InData, size_t InBlockSize, size_t InNumBlocks, void* InUserData)
+	{ return static_cast<FElectraHTTPStreamRequestLibCurl*>(InUserData)->CurlResponseBodyCallback(InData, InBlockSize, InNumBlocks); }
+	virtual size_t CurlUploadCallback(void* OutData, size_t InBlockSize, size_t InNumBlocks);
+	static size_t _CurlUploadCallback(void* OutData, size_t InBlockSize, size_t InNumBlocks, void* InUserData)
+	{ return static_cast<FElectraHTTPStreamRequestLibCurl*>(InUserData)->CurlUploadCallback(OutData, InBlockSize, InNumBlocks); }
+	virtual int CurlSeekCallback(curl_off_t InOffset, int InOrigin);
+	static int _CurlSeekCallback(void* InUserData, curl_off_t InOffset, int InOrigin)
+	{ return static_cast<FElectraHTTPStreamRequestLibCurl*>(InUserData)->CurlSeekCallback(InOffset, InOrigin); }
+#if WITH_SSL
+	virtual CURLcode CurlSSLCtxCallback(CURL* InEasyHandle, void* InSSLCtx);
+	static CURLcode _CurlSSLCtxCallback(CURL* InEasyHandle, void* InSSLCtx, void* InUserData)
+	{ return static_cast<FElectraHTTPStreamRequestLibCurl*>(InUserData)->CurlSSLCtxCallback(InEasyHandle, InSSLCtx); }
+	static int SslCertVerify(int PreverifyOk, X509_STORE_CTX* Context);
+#endif
+
+	// User agent. Defaults to a global one but can be set with each individual request.
+	FString UserAgent;
+	// GET, HEAD or POST
+	FString Verb;
+	// URL to request
+	FString URL;
+	// Optional byte range. If set this must be a valid range string.
+	FString Range;
+	// Host-only part of URL
+	FString Host;
+	// Set to true to allow gzip/deflate for GET requests.
+	bool bAllowCompression = false;
+	// Additional headers to be sent with the request.
+	TMap<FString, FString> AdditionalHeaders;
+
+	// Configuration
+	int MaxRedirections = 10;
+	int ConnectionTimeoutMillis = 10000;
+#if ELECTRA_HTTPSTREAM_CURL_ALLOW_UNSAFE_CONNECTIONS_FOR_DEBUGGING
+	bool bAllowUnsafeConnectionsForDebugging = false;
+#endif
+
+	FCriticalSection NotificationLock;
+	FElectraHTTPStreamNotificationDelegate NotificationCallback;
+
+	// Curl handles
+	CURL* CurlEasyHandle = nullptr;
+	curl_slist* CurlHeaderList = nullptr;
+	char CurlErrorMessageBuffer[CURL_ERROR_SIZE+1] = {0};
+	CURLM* CurlMultiHandle = nullptr;
+
+	bool bIsRedirectionResponse = false;
+	TArray<FElectraHTTPStreamHeader> CurrentResponseHeaders;
+	FString CurrentHTTPResponseStatusLine;
+	bool bResponseHeadersParsed = false;
+	long HttpCode = 0;
+	long HttpVersion = CURL_HTTP_VERSION_NONE;
+	FString EffectiveURL;
+
+	volatile EState CurrentState = EState::Inactive;
+	std::atomic_bool bCancel { false };
+	std::atomic_bool bTerminated { false };
+	int64 LastReportedDownloadSize = 0;
+
+	FElectraHTTPStreamBuffer PostData;
+	TSharedPtr<FElectraHTTPStreamResponse, ESPMode::ThreadSafe> Response;
+};
+
+#endif // ELECTRA_HTTPSTREAM_LIBCURL
diff --git a/Engine/Plugins/Media/ElectraUtil/Source/ElectraHTTPStream/Private/Curl/CurlLogError.cpp b/Engine/Plugins/Media/ElectraUtil/Source/ElectraHTTPStream/Private/Curl/CurlLogError.cpp
new file mode 100644
--- /dev/null
+++ b/Engine/Plugins/Media/ElectraUtil/Source/ElectraHTTPStream/Private/Curl/CurlLogError.cpp
@@ -0,0 +1,32 @@
+// Copyright Epic Games, Inc. All Rights Reserved.
+
+#if ELECTRA_HTTPSTREAM_LIBCURL
+
+#include "Curl/CurlLogError.h"
+#include "ElectraHTTPStreamModule.h"
+
+
+
+namespace ElectraHTTPStreamLibCurl
+{
+
+FString GetErrorMessage(CURLcode ErrorCode)
+{
+	FString msg(curl_easy_strerror(ErrorCode));
+	return msg;
+}
+
+FString GetErrorMessage(CURLMcode ErrorCode)
+{
+	FString msg(curl_multi_strerror(ErrorCode));
+	return msg;
+}
+
+void LogError(const FString& Message)
+{
+	UE_LOG(LogElectraHTTPStream, Error, TEXT("%s"), *Message);
+}
+
+}
+
+#endif // ELECTRA_HTTPSTREAM_LIBCURL
diff --git a/Engine/Plugins/Media/ElectraUtil/Source/ElectraHTTPStream/Private/Curl/CurlLogError.h b/Engine/Plugins/Media/ElectraUtil/Source/ElectraHTTPStream/Private/Curl/CurlLogError.h
new file mode 100644
--- /dev/null
+++ b/Engine/Plugins/Media/ElectraUtil/Source/ElectraHTTPStream/Private/Curl/CurlLogError.h
@@ -0,0 +1,14 @@
+// Copyright Epic Games, Inc. All Rights Reserved.
+
+#pragma once
+
+#include "CoreMinimal.h"
+
+#include "CurlElectra.h"
+
+namespace ElectraHTTPStreamLibCurl
+{
+	FString GetErrorMessage(CURLcode ErrorCode);
+	FString GetErrorMessage(CURLMcode ErrorCode);
+	void LogError(const FString& Message);
+}
diff --git a/Engine/Plugins/Media/ElectraUtil/Source/ElectraHTTPStream/Private/ElectraHTTPStreamResponse.h b/Engine/Plugins/Media/ElectraUtil/Source/ElectraHTTPStream/Private/ElectraHTTPStreamResponse.h
new file mode 100644
--- /dev/null
+++ b/Engine/Plugins/Media/ElectraUtil/Source/ElectraHTTPStream/Private/ElectraHTTPStreamResponse.h
@@ -0,0 +1,122 @@
+// Copyright Epic Games, Inc. All Rights Reserved.
+
+#pragma once
+
+#include "ElectraHTTPStream.h"
+#include "ElectraHTTPStreamBuffer.h"
+
+class FElectraHTTPStreamResponse : public IElectraHTTPStreamResponse
+{
+public:
+	FElectraHTTPStreamResponse() = default;
+	virtual ~FElectraHTTPStreamResponse() = default;
+
+	virtual EStatus GetStatus() override
+	{ return CurrentStatus; }
+
+	virtual EState GetState() override
+	{ return CurrentState; }
+
+	virtual FString GetErrorMessage() override
+	{ 
+		FScopeLock lock(&ErrorLock);
+		return ErrorMessage; 
+	}
+
+	virtual int32 GetHTTPResponseCode() override
+	{ return HTTPResponseCode; }
+
+	virtual int64 GetNumResponseBytesReceived() override
+	{ return ResponseData.GetNumTotalBytesAdded(); }
+
+	virtual int64 GetNumRequestBytesSent() override
+	{ return 0;	}
+
+	virtual FString GetEffectiveURL() override
+	{ return EffectiveURL; }
+
+	virtual void GetAllHeaders(TArray<FElectraHTTPStreamHeader>& OutHeaders) override
+	{ OutHeaders = ResponseHeaders;	}
+
+	virtual FString GetHTTPStatusLine() override
+	{ return HTTPStatusLine; }
+	virtual FString GetContentLengthHeader() override
+	{ return ContentLength; }
+	virtual FString GetContentRangeHeader() override
+	{ return ContentRange; }
+	virtual FString GetAcceptRangesHeader() override
+	{ return AcceptRanges; }
+	virtual FString GetTransferEncodingHeader() override
+	{ return TransferEncoding; }
+	virtual FString GetContentTypeHeader() override
+	{ return ContentType; }
+
+	virtual IElectraHTTPStreamBuffer& GetResponseData() override
+	{ return ResponseData; }
+
+	virtual double GetTimeElapsed() override
+	{ return CurrentStatus == EStatus::NotRunning ? 0.0 : TimeUntilFinished > 0.0 ? TimeUntilFinished : FPlatformTime::Seconds() - StartTime; }
+
+	virtual double GetTimeSinceLastDataArrived() override
+	{ return CurrentStatus == EStatus::NotRunning || TimeOfMostRecentReceive == 0.0 || TimeUntilFinished > 0.0 ? 0.0 : FPlatformTime::Seconds() - TimeOfMostRecentReceive; }
+
+	virtual double GetTimeUntilNameResolved() override
+	{ return TimeUntilNameResolved; }
+	virtual double GetTimeUntilConnected() override
+	{ return TimeUntilConnected; }
+	virtual double GetTimeUntilRequestSent() override
+	{ return TimeUntilRequestSent; }
+	virtual double GetTimeUntilHeadersAvailable() override
+	{ return TimeUntilHeadersAvailable; }
+	virtual double GetTimeUntilFirstByte() override
+	{ return TimeUntilFirstByte; }
+	virtual double GetTimeUntilFinished() override
+	{ return TimeUntilFinished; }
+
+	void AddResponseData(const TArray<uint8>& InNewData)
+	{ ResponseData.AddData(InNewData); }
+	void AddResponseData(TArray<uint8>&& InNewData)
+	{ ResponseData.AddData(MoveTemp(InNewData)); }
+	void AddResponseData(const TConstArrayView<const uint8>& InNewData)
+	{ ResponseData.AddData(InNewData); }
+
+
+	void SetEOS()
+	{ ResponseData.SetEOS(); }
+
+	void SetErrorMessage(const FString& InErrorMessage)
+	{
+		FScopeLock lock(&ErrorLock);
+		ErrorMessage = InErrorMessage;
+	}
+
+	EStatus CurrentStatus = EStatus::NotRunning;
+	EState CurrentState = EState::Connecting;
+
+	TArray<FElectraHTTPStreamHeader> ResponseHeaders;
+	FString ContentLength;
+	FString ContentRange;
+	FString AcceptRanges;
+	FString TransferEncoding;
+	FString ContentType;
+
+	FString HTTPStatusLine;
+	FString EffectiveURL;
+	int32 HTTPResponseCode = 0;
+
+	FElectraHTTPStreamBuffer ResponseData;
+
+	// Timing
+	double StartTime = 0.0;
+	double TimeUntilNameResolved = 0.0;
+	double TimeUntilConnected = 0.0;
+	double TimeUntilRequestSent = 0.0;
+	double TimeUntilHeadersAvailable = 0.0;
+	double TimeUntilFirstByte = 0.0;
+	double TimeUntilFinished = 0.0;
+	double TimeOfMostRecentReceive = 0.0;
+
+	// Error message.
+	FCriticalSection ErrorLock;
+	FString ErrorMessage;
+};
