ElectraHTTPStream: Added a network module geared towards downloading of chunked media streams for use with the Electra media player.

diff --git a/Engine/Plugins/Media/ElectraUtil/Source/ElectraHTTPStream/Private/WinHttp/WinHttpElectraHTTPStream.cpp b/Engine/Plugins/Media/ElectraUtil/Source/ElectraHTTPStream/Private/WinHttp/WinHttpElectraHTTPStream.cpp
new file mode 100644
--- /dev/null
+++ b/Engine/Plugins/Media/ElectraUtil/Source/ElectraHTTPStream/Private/WinHttp/WinHttpElectraHTTPStream.cpp
@@ -0,0 +1,1404 @@
+// Copyright Epic Games, Inc. All Rights Reserved.
+
+#if ELECTRA_HTTPSTREAM_WINHTTP
+
+#include "WinHttp/WinHttpElectraHTTPStream.h"
+#include "HAL/ThreadSafeCounter.h"
+#include "HAL/PlatformTime.h"
+#include "HAL/PlatformProcess.h"
+#include "HAL/Runnable.h"
+#include "HAL/RunnableThread.h"
+#include "HAL/LowLevelMemTracker.h"
+
+#include "ElectraHTTPStreamModule.h"
+#include "ElectraHTTPStream.h"
+#include "ElectraHTTPStreamBuffer.h"
+#include "ElectraHTTPStreamResponse.h"
+#include "ElectraHTTPStreamPerf.h"
+
+#include "Utilities/TimeWaitableSignal.h"
+#include "Utilities/DefaultHttpUserAgent.h"
+#include "Utilities/HttpRangeHeader.h"
+#include "Utilities/URLParser.h"
+#include "WinHttp/WinHttpLogError.h"
+
+#pragma warning(push)
+#pragma warning(disable : 28285) // Disable static analysis syntax error in WinHttpSetHeaders macros
+#include "Windows/AllowWindowsPlatformTypes.h"
+THIRD_PARTY_INCLUDES_START
+#include <errhandlingapi.h>
+#include <winhttp.h>
+THIRD_PARTY_INCLUDES_END
+#include "Windows/HideWindowsPlatformTypes.h"
+#pragma warning(pop)
+
+
+#include <atomic>
+
+// Size WinHTTP uses internally for its read buffer, as per the documentation.
+#define ELECTRA_HTTPSTREAM_INTERNAL_WINHTTP_READ_BUFFER_SIZE	8192
+
+// Check if the header files say HTTP/2 could be enabled.
+#if defined(WINHTTP_PROTOCOL_FLAG_HTTP2) && defined(WINHTTP_OPTION_ENABLE_HTTP_PROTOCOL)
+#define ELECTRA_HTTPSTREAM_INTERNAL_WINHTTP_ENABLE_HTTP2 1
+#endif
+
+// In all but shipping builds we allow disabling of security checks.
+#if !UE_BUILD_SHIPPING
+#define ELECTRA_HTTPSTREAM_WINHTTP_ALLOW_UNSAFE_CONNECTIONS_FOR_DEBUGGING 1
+#endif
+
+#if ELECTRA_HTTPSTREAM_REQUIRES_SECURE_CONNECTIONS && !defined(WINHTTP_FLAG_SECURE_DEFAULTS)
+#undef ELECTRA_HTTPSTREAM_REQUIRES_SECURE_CONNECTIONS
+#define ELECTRA_HTTPSTREAM_REQUIRES_SECURE_CONNECTIONS 0
+#endif
+
+
+/***************************************************************************************************************************************************/
+/***************************************************************************************************************************************************/
+/***************************************************************************************************************************************************/
+
+DECLARE_STATS_GROUP(TEXT("Electra HTTP Stream"), STATGROUP_ElectraHTTPStream, STATCAT_Advanced);
+DECLARE_CYCLE_STAT(TEXT("Process"), STAT_ElectraHTTPThread_Process, STATGROUP_ElectraHTTPStream);
+DECLARE_CYCLE_STAT(TEXT("AsyncCallback"), STAT_ElectraHTTPThread_AsyncCallback, STATGROUP_ElectraHTTPStream);
+DECLARE_CYCLE_STAT(TEXT("Custom handler"), STAT_ElectraHTTPThread_CustomHandler, STATGROUP_ElectraHTTPStream);
+
+/***************************************************************************************************************************************************/
+/***************************************************************************************************************************************************/
+/***************************************************************************************************************************************************/
+
+class FElectraHTTPStreamRequestWinHttp;
+
+/***************************************************************************************************************************************************/
+/***************************************************************************************************************************************************/
+/***************************************************************************************************************************************************/
+
+
+/**
+ * WinHttp version of the ElectraHTTPStream implementation.
+ */
+class FElectraHTTPStreamWinHttp : public TSharedFromThis<FElectraHTTPStreamWinHttp, ESPMode::ThreadSafe>, public IElectraHTTPStream, private FRunnable
+{
+public:
+	virtual ~FElectraHTTPStreamWinHttp();
+
+	FElectraHTTPStreamWinHttp();
+	bool Initialize(const Electra::FParamDict& InOptions);
+
+	static uint64 GetNewRequestIndex()
+	{ return ++NextRequestIndex; }
+
+	virtual FElectraHTTPStreamThreadHandlerDelegate& ThreadHandlerDelegate() override
+	{ return ThreadHandlerCallback; }
+
+	virtual void Close() override;
+
+	virtual IElectraHTTPStreamRequestPtr CreateRequest() override;
+
+	virtual void AddRequest(IElectraHTTPStreamRequestPtr Request) override;
+
+	HINTERNET GetSessionHandle()
+	{ return SessionHandle; }
+
+	void TriggerWorkSignal()
+	{ HaveWorkSignal.Signal(); }
+
+	static void _WinHttpCallback(HINTERNET hInternet, DWORD_PTR dwContext, DWORD dwInternetStatus, void* lpvStatusInformation, DWORD dwStatusInformationLength);
+
+private:
+	struct FRequestPointers
+	{
+		TWeakPtr<FElectraHTTPStreamWinHttp, ESPMode::ThreadSafe> Owner;
+		TWeakPtr<FElectraHTTPStreamRequestWinHttp, ESPMode::ThreadSafe>	Request;
+	};
+
+	// Methods from FRunnable
+	virtual uint32 Run() override final;
+	virtual void Stop() override final;
+
+
+	static FRequestPointers GetRequestByIndex(uint64 InRequestIndex)
+	{
+		FRequestPointers Pointers;
+		FScopeLock lock(&AllRequestHandleLock);
+		FRequestPointers* Ptrs = AllRequestHandles.Find(InRequestIndex);
+		if (Ptrs)
+		{
+			Pointers = *Ptrs;
+		}
+		return Pointers;
+	}
+
+	static void RemoveOutdatedRequests()
+	{
+		FScopeLock lock(&AllRequestHandleLock);
+		TArray<uint64> OutdatedIndices;
+		for(auto &Req : AllRequestHandles)
+		{
+			if (!Req.Value.Request.IsValid() || !Req.Value.Owner.IsValid())
+			{
+				OutdatedIndices.Emplace(Req.Key);
+			}
+		}
+		for(auto &Index : OutdatedIndices)
+		{
+			AllRequestHandles.Remove(Index);
+		}
+	}
+
+	void SetupNewRequests();
+	void UpdateActiveRequests();
+	void HandleCompletedRequests();
+
+
+	HINTERNET SessionHandle = NULL;
+
+	FThreadSafeCounter ExitRequest;
+	FRunnableThread* Thread = nullptr;
+	FTimeWaitableSignal HaveWorkSignal;
+
+	FElectraHTTPStreamThreadHandlerDelegate ThreadHandlerCallback;
+
+	FCriticalSection RequestLock;
+	TArray<TSharedPtr<FElectraHTTPStreamRequestWinHttp, ESPMode::ThreadSafe>> NewRequests;
+	TArray<TSharedPtr<FElectraHTTPStreamRequestWinHttp, ESPMode::ThreadSafe>> ActiveRequests;
+	TArray<TSharedPtr<FElectraHTTPStreamRequestWinHttp, ESPMode::ThreadSafe>> CompletedRequests;
+
+
+	// Static members. These are used to track requests by a unique integer value instead of any pointers,
+	// not even shared/weak pointers. This reduces issues with cyclic references and possible async WinHttp
+	// callbacks occurring for already removed requests.
+	static std::atomic_ulong NextRequestIndex;
+	static FCriticalSection AllRequestHandleLock;
+	static TMap<uint64, FRequestPointers> AllRequestHandles;
+};
+std::atomic_ulong											FElectraHTTPStreamWinHttp::NextRequestIndex = {0};
+FCriticalSection											FElectraHTTPStreamWinHttp::AllRequestHandleLock;
+TMap<uint64, FElectraHTTPStreamWinHttp::FRequestPointers>	FElectraHTTPStreamWinHttp::AllRequestHandles;
+
+/***************************************************************************************************************************************************/
+/***************************************************************************************************************************************************/
+/***************************************************************************************************************************************************/
+
+/**
+ * WinHttp version of a HTTP stream request.
+ */
+class FElectraHTTPStreamRequestWinHttp : public TSharedFromThis<FElectraHTTPStreamRequestWinHttp, ESPMode::ThreadSafe>, public IElectraHTTPStreamRequest
+{
+public:
+	enum EState
+	{
+		Inactive,
+		Started,
+		Connecting,
+		SendingRequest,
+		RequestSent,
+		AwaitingResponse,
+		ResponseReceived,
+		HeadersAvailable,
+		ParsingHeaders,
+		AwaitingResponseData,
+		ResponseDataAvailable,
+		ReadingResponseData,
+		ResponseDataRead,
+		Finished,
+		Error,
+		NeedsRedirect
+	};
+
+	enum EReadResponseResult
+	{
+		Success,
+		Failed,
+		EndOfData
+	};
+
+	FElectraHTTPStreamRequestWinHttp(uint64 InRequestIndex);
+	virtual ~FElectraHTTPStreamRequestWinHttp();
+
+	virtual void SetVerb(const FString& InVerb) override
+	{ Verb = InVerb; }
+
+	virtual IElectraHTTPStreamBuffer& POSTDataBuffer() override
+	{ return PostData; }
+
+	virtual void SetUserAgent(const FString& InUserAgent) override
+	{ UserAgent = InUserAgent; }
+
+	virtual void SetURL(const FString& InURL) override
+	{ URL = InURL; }
+	virtual void SetRange(const FString& InRange) override
+	{ Range = InRange; }
+	virtual void AllowCompression(bool bInAllowCompression) override
+	{ bAllowCompression = bInAllowCompression; }
+
+	virtual void AllowUnsafeRequestsForDebugging() override
+	{
+	#ifdef ELECTRA_HTTPSTREAM_WINHTTP_ALLOW_UNSAFE_CONNECTIONS_FOR_DEBUGGING
+		bAllowUnsafeConnectionsForDebugging = true;
+	#endif
+	}
+
+	virtual void AddHeader(const FString& Header, const FString& Value, bool bAppendIfExists) override
+	{
+		// Ignore a few headers that will be set automatically.
+		if (Header.Equals(TEXT("User-Agent"), ESearchCase::IgnoreCase) ||
+			Header.Equals(TEXT("Accept-Encoding"), ESearchCase::IgnoreCase) ||
+			Header.Equals(TEXT("Transfer-Encoding"), ESearchCase::IgnoreCase) ||
+			Header.Equals(TEXT("Range"), ESearchCase::IgnoreCase) ||
+			Header.Equals(TEXT("Accept-Ranges"), ESearchCase::IgnoreCase))
+		{
+			return;
+		}
+		FString* ExistingHeader = AdditionalHeaders.Find(Header);
+		if (ExistingHeader)
+		{
+			// If the value is empty and we shall replace the header we need to remove it since empty headers are prohibited.
+			if (Value.IsEmpty() && !bAppendIfExists)
+			{
+				AdditionalHeaders.Remove(Header);
+			}
+			else if (Value.Len())
+			{
+				FString NewValue = *ExistingHeader;
+				NewValue.Append(TEXT(", "));
+				NewValue.Append(Value);
+				*ExistingHeader = NewValue;
+			}
+		}
+		// Header does not exist yet. If the value is empty we must not add the header since empty headers are prohibited.
+		else if (Value.Len())
+		{
+			AdditionalHeaders.Add(Header, Value);
+		}
+	}
+
+	virtual FElectraHTTPStreamNotificationDelegate& NotificationDelegate() override
+	{ return NotificationCallback; }
+
+	virtual void Cancel() override
+	{
+		FScopeLock lock(&NotificationLock);
+		NotificationCallback.Unbind();
+		bCancel = true;
+	}
+
+	virtual IElectraHTTPStreamResponsePtr GetResponse() override
+	{ return Response; }
+
+	virtual bool HasFailed() override
+	{ return Response->GetErrorMessage().Len() > 0; }
+
+	virtual FString GetErrorMessage() override
+	{ return Response->GetErrorMessage(); }
+
+	uint64 GetRequestIndex()
+	{ return RequestIndex; }
+
+	EState GetCurrentState()
+	{ return CurrentState; }
+
+	bool WasCanceled()
+	{ return bCancel; }
+
+	bool Setup(FElectraHTTPStreamWinHttp* OwningManager);
+	bool Execute();
+	bool RequestResponse();
+	bool ParseHeaders(bool bNotify);
+	bool RequestResponseData();
+	EReadResponseResult ReadResponseData();
+	bool HandleNewResponseData();
+	void SetFinished();
+
+	void SetResponseStatus(IElectraHTTPStreamResponse::EStatus InStatus)
+	{ Response->CurrentStatus = InStatus; }
+
+	void NotifyCallback(EElectraHTTPStreamNotificationReason InReason, int64 InParam)
+	{
+		FScopeLock lock(&NotificationLock);
+		NotificationCallback.ExecuteIfBound(AsShared(), InReason, InParam);
+	}
+
+	void Close(bool bHandlesOnly=false);
+	void Terminate();
+
+	void WinHttpCallback(FElectraHTTPStreamWinHttp* Owner, HINTERNET hInternet, DWORD dwInternetStatus, void* lpvStatusInformation, DWORD dwStatusInformationLength);
+private:
+	FElectraHTTPStreamRequestWinHttp() = delete;
+	FElectraHTTPStreamRequestWinHttp(const FElectraHTTPStreamRequestWinHttp&) = delete;
+	FElectraHTTPStreamRequestWinHttp& operator=(const FElectraHTTPStreamRequestWinHttp&) = delete;
+
+	// Unique request ID.
+	uint64 RequestIndex = 0;
+
+	// User agent. Defaults to a global one but can be set with each individual request.
+	FString UserAgent;
+	// GET or POST
+	FString Verb;
+	// URL to request
+	FString URL;
+	// Optional byte range. If set this must be a valid range string.
+	FString Range;
+	// Set to true to allow gzip/deflate for GET requests.
+	bool bAllowCompression = false;
+	// Additional headers to be sent with the request.
+	TMap<FString, FString> AdditionalHeaders;
+
+	// Configuration
+	int MaxRedirections = 10;
+	#ifdef ELECTRA_HTTPSTREAM_WINHTTP_ALLOW_UNSAFE_CONNECTIONS_FOR_DEBUGGING
+	bool bAllowUnsafeConnectionsForDebugging = false;
+	#endif
+
+	FCriticalSection NotificationLock;
+	FElectraHTTPStreamNotificationDelegate NotificationCallback;
+
+	// Handles as required by WinHttp
+	HINTERNET ConnectionHandle = NULL;
+	HINTERNET RequestHandle = NULL;
+	DWORD LastWinHttpCallbackReason = 0;
+	FString SecureFailureFlags;
+
+	FString RedirectedLocation;
+	std::atomic_bool bMustRedirect { false };
+
+	volatile EState CurrentState = EState::Inactive;
+	std::atomic_bool bCancel { false };
+	uint32 NumResponseBytesAvailableToRead = 0;
+	TArray<uint8> ChunkReadBuffer;
+
+	FElectraHTTPStreamBuffer PostData;
+	TSharedPtr<FElectraHTTPStreamResponse, ESPMode::ThreadSafe> Response;
+};
+
+/***************************************************************************************************************************************************/
+/***************************************************************************************************************************************************/
+/***************************************************************************************************************************************************/
+
+FElectraHTTPStreamRequestWinHttp::FElectraHTTPStreamRequestWinHttp(uint64 InRequestIndex)
+	: RequestIndex(InRequestIndex)
+{
+	UserAgent = ElectraHTTPStream::GetDefaultUserAgent();
+	Response = MakeShared<FElectraHTTPStreamResponse, ESPMode::ThreadSafe>();
+}
+
+FElectraHTTPStreamRequestWinHttp::~FElectraHTTPStreamRequestWinHttp()
+{
+	Close();
+}
+
+bool FElectraHTTPStreamRequestWinHttp::Setup(FElectraHTTPStreamWinHttp* OwningManager)
+{
+	bool bIsRedirect = bMustRedirect;
+	if (bIsRedirect)
+	{
+		bMustRedirect = false;
+		URL = RedirectedLocation;
+		RedirectedLocation.Empty();
+		RequestIndex = FElectraHTTPStreamWinHttp::GetNewRequestIndex();
+		CurrentState = EState::Inactive;
+	}
+
+	// Check for a supported verb.
+	if (Verb.IsEmpty())
+	{
+		Verb = TEXT("GET");
+	}
+	if (!(Verb.Equals(TEXT("GET")) || Verb.Equals(TEXT("POST")) || Verb.Equals(TEXT("HEAD"))))
+	{
+		Response->SetErrorMessage(FString::Printf(TEXT("Unsupported verb \"%s\""), *Verb));
+		ElectraHTTPStreamWinHttp::LogError(Response->GetErrorMessage());
+		return false;
+	}
+
+	Electra::FURL_RFC3986 UrlParser;
+	if (!UrlParser.Parse(URL))
+	{
+		Response->SetErrorMessage(FString::Printf(TEXT("Failed to parse URL \"%s\""), *URL));
+		ElectraHTTPStreamWinHttp::LogError(Response->GetErrorMessage());
+		return false;
+	}
+	INTERNET_PORT Port = INTERNET_DEFAULT_PORT;
+	FString UrlPort = UrlParser.GetPort();
+	if (UrlPort.IsEmpty())
+	{
+		UrlPort = Electra::FURL_RFC3986::GetStandardPortForScheme(UrlParser.GetScheme());
+	}
+	if (UrlPort.Len())
+	{
+		int64 p = 0;
+		LexFromString(p, *UrlPort);
+		Port = static_cast<INTERNET_PORT>(p);
+	}
+	HINTERNET ch = WinHttpConnect(OwningManager->GetSessionHandle(), TCHAR_TO_WCHAR(*UrlParser.GetHost()), Port, 0);
+	if (ch == NULL)
+	{
+		Response->SetErrorMessage(ElectraHTTPStreamWinHttp::GetErrorLogMessage(TEXT("WinHttpConnect()"), GetLastError()));
+		ElectraHTTPStreamWinHttp::LogError(Response->GetErrorMessage());
+		return false;
+	}
+	ConnectionHandle = ch;
+
+	bool bIsSecure = UrlParser.GetScheme().Equals(TEXT("https"), ESearchCase::IgnoreCase);
+#if ELECTRA_HTTPSTREAM_REQUIRES_SECURE_CONNECTIONS
+	if (!bIsSecure)
+	{
+		// Once the session has been set to require secure connections it is not possible to perform a http: connection,
+		// not even if we wanted to, so we have to error out here.
+		Response->SetErrorMessage(FString::Printf(TEXT("Attempted to create an insecure http:// request which is disabled on this platform")));
+		ElectraHTTPStreamWinHttp::LogError(Response->GetErrorMessage());
+		return false;
+	}
+#endif
+
+
+	FString Path = UrlParser.GetPath(true, false);
+	DWORD Flags = WINHTTP_FLAG_ESCAPE_DISABLE | WINHTTP_FLAG_ESCAPE_DISABLE_QUERY;
+	if (bIsSecure)
+	{
+		Flags |= WINHTTP_FLAG_SECURE;
+	}
+	HINTERNET rh = WinHttpOpenRequest(ConnectionHandle, TCHAR_TO_WCHAR(*Verb), TCHAR_TO_WCHAR(*Path), NULL, WINHTTP_NO_REFERER, WINHTTP_DEFAULT_ACCEPT_TYPES, Flags);
+	if (rh == NULL)
+	{
+		Response->SetErrorMessage(ElectraHTTPStreamWinHttp::GetErrorLogMessage(TEXT("WinHttpOpenRequest()"), GetLastError()));
+		ElectraHTTPStreamWinHttp::LogError(Response->GetErrorMessage());
+		return false;
+	}
+	RequestHandle = rh;
+
+	// Set options on the request handle.
+	DWORD OptionValue;
+
+	// Disable cookies.
+	OptionValue = WINHTTP_DISABLE_COOKIES;
+	if (!WinHttpSetOption(RequestHandle, WINHTTP_OPTION_DISABLE_FEATURE, &OptionValue, sizeof(OptionValue)))
+	{
+		Response->SetErrorMessage(ElectraHTTPStreamWinHttp::GetErrorLogMessage(TEXT("WinHttpSetOption(WINHTTP_OPTION_DISABLE_FEATURE:WINHTTP_DISABLE_COOKIES)"), GetLastError()));
+		ElectraHTTPStreamWinHttp::LogError(Response->GetErrorMessage());
+		return false;
+	}
+	// Client certificates are not supported.
+	if (bIsSecure)
+	{
+		// This may only be used with secure connections. Otherwise the call will fail.
+		if (!WinHttpSetOption(RequestHandle, WINHTTP_OPTION_CLIENT_CERT_CONTEXT, WINHTTP_NO_CLIENT_CERT_CONTEXT, 0))
+		{
+			Response->SetErrorMessage(ElectraHTTPStreamWinHttp::GetErrorLogMessage(TEXT("WinHttpSetOption(WINHTTP_OPTION_CLIENT_CERT_CONTEXT:WINHTTP_NO_CLIENT_CERT_CONTEXT)"), GetLastError()));
+			ElectraHTTPStreamWinHttp::LogError(Response->GetErrorMessage());
+			return false;
+		}
+	}
+
+	// For debugging purposes we may disable certificate validation for selected requests.
+	#ifdef ELECTRA_HTTPSTREAM_WINHTTP_ALLOW_UNSAFE_CONNECTIONS_FOR_DEBUGGING
+	if (bAllowUnsafeConnectionsForDebugging)
+	{
+		OptionValue = SECURITY_FLAG_IGNORE_UNKNOWN_CA;
+		bool bAllowUnsafe = WinHttpSetOption(RequestHandle, WINHTTP_OPTION_SECURITY_FLAGS, &OptionValue, sizeof(OptionValue));
+		check(bAllowUnsafe);
+	}
+	#endif
+
+	// We set the user agent manually. It can only be set on the session handle but we want to support different agents
+	// on a per-request basis if necessary. For simplicities sake we add it to the list of additional headers.
+	AdditionalHeaders.FindOrAdd(TEXT("User-Agent"), UserAgent);
+
+	// GET/HEAD / POST specific options
+	if (Verb.Equals("GET") || Verb.Equals("HEAD"))
+	{
+		if (Range.Len())
+		{
+			FString RangeHdr = Range;
+			if (!RangeHdr.StartsWith(TEXT("bytes=")))
+			{
+				RangeHdr.InsertAt(0, TEXT("bytes="));
+			}
+			AdditionalHeaders.FindOrAdd(TEXT("Range"), RangeHdr);
+		}
+
+		/*
+			Removed because WinHttp has a problem if the response is compressed AND chunked, with the
+			response header being "Transfer-Encoding: gzip,chunked"
+			If this happens WinHttp does not decompress and stalls after having received the compressed size.
+
+		if (bAllowCompression)
+		{
+			DWORD CompressionTypes = WINHTTP_DECOMPRESSION_FLAG_ALL;
+			if (WinHttpSetOption(RequestHandle, WINHTTP_OPTION_DECOMPRESSION, &CompressionTypes, sizeof(CompressionTypes)))
+			{
+				AdditionalHeaders.FindOrAdd(TEXT("Accept-Encoding"), TEXT("deflate, gzip"));
+			}
+			else
+			{
+				// Compression was added in Windows 8.1. Earlier versions do not support compression, so it is ok if this fails.
+			}
+		}
+		*/
+	}
+
+	// Add all additional headers.
+	for(auto &Hdr : AdditionalHeaders)
+	{
+		auto BuildHeader = [](const FString& InHdr, const FString& InVal) -> FString
+		{
+			if (InHdr.Len() && InVal.Len())
+			{
+				return InHdr + TEXT(": ") + InVal;
+			}
+			return FString();
+		};
+
+		FString Header = BuildHeader(Hdr.Key, Hdr.Value);
+		if (Header.Len())
+		{
+			const DWORD HeadersLength = Header.Len();
+			const DWORD HeaderFlags = WINHTTP_ADDREQ_FLAG_ADD | WINHTTP_ADDREQ_FLAG_REPLACE;
+			if (!WinHttpAddRequestHeaders(RequestHandle, TCHAR_TO_WCHAR(*Header), HeadersLength, HeaderFlags))
+			{
+				Response->SetErrorMessage(ElectraHTTPStreamWinHttp::GetErrorLogMessage(TEXT("WinHttpAddRequestHeaders()"), GetLastError()));
+				ElectraHTTPStreamWinHttp::LogError(Response->GetErrorMessage());
+				return false;
+			}
+		}
+	}
+
+	return true;
+}
+
+bool FElectraHTTPStreamRequestWinHttp::Execute()
+{
+	// Set the origin URL as effective URL first in case there are problems or no redirections.
+	Response->EffectiveURL = URL;
+	Response->StartTime = FPlatformTime::Seconds();
+	Response->CurrentStatus = IElectraHTTPStreamResponse::EStatus::Running;
+
+	CurrentState = EState::Started;
+
+	// Register callback for this request.
+	WINHTTP_STATUS_CALLBACK PreviousCallback = WinHttpSetStatusCallback(RequestHandle, &FElectraHTTPStreamWinHttp::_WinHttpCallback, WINHTTP_CALLBACK_FLAG_ALL_NOTIFICATIONS, 0);
+	if (PreviousCallback == WINHTTP_INVALID_STATUS_CALLBACK)
+	{
+		Response->SetErrorMessage(ElectraHTTPStreamWinHttp::GetErrorLogMessage(TEXT("WinHttpSetStatusCallback()"), GetLastError()));
+		ElectraHTTPStreamWinHttp::LogError(Response->GetErrorMessage());
+		return false;
+	}
+
+	if (Verb.Equals("GET") || Verb.Equals("HEAD"))
+	{
+		if (!WinHttpSendRequest(RequestHandle, WINHTTP_NO_ADDITIONAL_HEADERS, 0, WINHTTP_NO_REQUEST_DATA, 0, 0, RequestIndex))
+		{
+			Response->SetErrorMessage(ElectraHTTPStreamWinHttp::GetErrorLogMessage(TEXT("WinHttpSendRequest()"), GetLastError()));
+			ElectraHTTPStreamWinHttp::LogError(Response->GetErrorMessage());
+			return false;
+		}
+	}
+	else if (Verb.Equals("POST"))
+	{
+		// For now we need the EOS flag set as we send the data as a whole.
+		check(PostData.GetEOS());
+
+		const uint8* DataToSend;
+		int64 NumDataToSend;
+		PostData.LockBuffer(DataToSend, NumDataToSend);
+		PostData.UnlockBuffer(0);
+		if (!WinHttpSendRequest(RequestHandle, WINHTTP_NO_ADDITIONAL_HEADERS, 0, (LPVOID)DataToSend, (DWORD)NumDataToSend, (DWORD)NumDataToSend, RequestIndex))
+		{
+			Response->SetErrorMessage(ElectraHTTPStreamWinHttp::GetErrorLogMessage(TEXT("WinHttpSendRequest()"), GetLastError()));
+			ElectraHTTPStreamWinHttp::LogError(Response->GetErrorMessage());
+			return false;
+		}
+	}
+	return true;
+}
+
+bool FElectraHTTPStreamRequestWinHttp::RequestResponse()
+{
+	CurrentState = EState::AwaitingResponse;
+	Response->CurrentState = IElectraHTTPStreamResponse::EState::WaitingForResponseHeaders;
+
+	if (!WinHttpReceiveResponse(RequestHandle, NULL))
+	{
+		Response->SetErrorMessage(ElectraHTTPStreamWinHttp::GetErrorLogMessage(TEXT("WinHttpReceiveResponse()"), GetLastError()));
+		ElectraHTTPStreamWinHttp::LogError(Response->GetErrorMessage());
+		return false;
+	}
+	return true;
+}
+
+bool FElectraHTTPStreamRequestWinHttp::ParseHeaders(bool bNotify)
+{
+	CurrentState = EState::ParsingHeaders;
+
+	// First ask how large the buffer we have to allocate needs to be.
+	DWORD RequiredHeaderSize = 0;
+	// This method HAS TO fail and give us the required size. If it succeeds for whatever reason that would be wrong.
+	if (WinHttpQueryHeaders(RequestHandle, WINHTTP_QUERY_RAW_HEADERS, WINHTTP_HEADER_NAME_BY_INDEX, WINHTTP_NO_OUTPUT_BUFFER, &RequiredHeaderSize, WINHTTP_NO_HEADER_INDEX))
+	{
+		Response->SetErrorMessage(ElectraHTTPStreamWinHttp::GetErrorLogMessage(TEXT("WinHttpQueryHeaders(); query for size succeeded where it should have failed!"), GetLastError()));
+		ElectraHTTPStreamWinHttp::LogError(Response->GetErrorMessage());
+		return false;
+	}
+	DWORD ErrorCode = GetLastError();
+	if (ErrorCode != ERROR_INSUFFICIENT_BUFFER)
+	{
+		Response->SetErrorMessage(ElectraHTTPStreamWinHttp::GetErrorLogMessage(TEXT("WinHttpQueryHeaders()"), ErrorCode));
+		ElectraHTTPStreamWinHttp::LogError(Response->GetErrorMessage());
+		return false;
+	}
+	// Are there headers to get?
+	if (RequiredHeaderSize)
+	{
+		TArray<wchar_t> AllHeadersBuffer;
+		AllHeadersBuffer.SetNumUninitialized(RequiredHeaderSize / sizeof(wchar_t), false);
+
+		// Repeat the header query, this time with the buffer to read the headers into.
+		if (!WinHttpQueryHeaders(RequestHandle, WINHTTP_QUERY_RAW_HEADERS, WINHTTP_HEADER_NAME_BY_INDEX, AllHeadersBuffer.GetData(), &RequiredHeaderSize, WINHTTP_NO_HEADER_INDEX))
+		{
+			Response->SetErrorMessage(ElectraHTTPStreamWinHttp::GetErrorLogMessage(TEXT("WinHttpQueryHeaders()"), GetLastError()));
+			ElectraHTTPStreamWinHttp::LogError(Response->GetErrorMessage());
+			return false;
+		}
+
+		// Getting RAW headers means each is terminated by a NUL character with the end of them being an empty string of just a NUL character.
+		int32 HdrPos = 0, HdrLen = 0;
+		while((HdrLen = FCStringWide::Strlen(AllHeadersBuffer.GetData() + HdrPos)) != 0)
+		{
+			FWideStringView FullHeader(AllHeadersBuffer.GetData() + HdrPos, HdrLen + 1);
+			HdrPos += HdrLen + 1;
+
+			int32 ColonPos = INDEX_NONE;
+			if (FullHeader.FindChar(L':', ColonPos))
+			{
+				FWideStringView HeaderKey(FullHeader.Left(ColonPos));
+				FWideStringView HeaderValue(FullHeader.RightChop(ColonPos + 1));
+
+				// Remove trailing NUL terminator from the view, otherwise the resulting string will end up double terminated.
+				HeaderValue.RemoveSuffix(1);
+				HeaderValue.TrimStartAndEndInline();
+
+				FElectraHTTPStreamHeader Hdr;
+				Hdr.Header = HeaderKey;
+				Hdr.Value = HeaderValue;
+
+				// Check for the most commonly used headers and set them on the side for faster access.
+				if (Hdr.Header.Equals(TEXT("Content-Type"), ESearchCase::IgnoreCase))
+				{
+					Response->ContentType = Hdr.Value;
+				}
+				else if (Hdr.Header.Equals(TEXT("Content-Length"), ESearchCase::IgnoreCase))
+				{
+					Response->ContentLength = Hdr.Value;
+				}
+				else if (Hdr.Header.Equals(TEXT("Content-Range"), ESearchCase::IgnoreCase))
+				{
+					Response->ContentRange = Hdr.Value;
+				}
+				else if (Hdr.Header.Equals(TEXT("Accept-Ranges"), ESearchCase::IgnoreCase))
+				{
+					Response->AcceptRanges = Hdr.Value;
+				}
+				else if (Hdr.Header.Equals(TEXT("Transfer-Encoding"), ESearchCase::IgnoreCase))
+				{
+					Response->TransferEncoding = Hdr.Value;
+				}
+
+				// Add to the list of headers, even those we have treated separately.
+				Response->ResponseHeaders.Emplace(MoveTemp(Hdr));
+			}
+			else
+			{
+				// No colon, this should be the HTTP status line, eg. "HTTP/1.1 200 OK"
+				Response->HTTPStatusLine = FullHeader;
+			}
+		}
+
+		// We could parse the HTTP status line ourselves, but for simplicities sake we ask WinHttp for it.
+		DWORD StatusCode = 0;
+		RequiredHeaderSize = sizeof(StatusCode);
+		if (!WinHttpQueryHeaders(RequestHandle, WINHTTP_QUERY_STATUS_CODE | WINHTTP_QUERY_FLAG_NUMBER, WINHTTP_HEADER_NAME_BY_INDEX, &StatusCode, &RequiredHeaderSize, WINHTTP_NO_HEADER_INDEX))
+		{
+			Response->SetErrorMessage(ElectraHTTPStreamWinHttp::GetErrorLogMessage(TEXT("WinHttpQueryHeaders()"), GetLastError()));
+			ElectraHTTPStreamWinHttp::LogError(Response->GetErrorMessage());
+			return false;
+		}
+		Response->HTTPResponseCode = static_cast<int32>(StatusCode);
+
+		// Get the effective URL
+		DWORD RequiredBufferSize = 0;
+		if (!WinHttpQueryOption(RequestHandle, WINHTTP_OPTION_URL, NULL, &RequiredBufferSize))
+		{
+			TArray<wchar_t> UrlBuffer;
+			UrlBuffer.SetNumUninitialized(RequiredBufferSize / sizeof(wchar_t), false);
+			if (!WinHttpQueryOption(RequestHandle, WINHTTP_OPTION_URL, UrlBuffer.GetData(), &RequiredBufferSize))
+			{
+				Response->SetErrorMessage(ElectraHTTPStreamWinHttp::GetErrorLogMessage(TEXT("WinHttpQueryOption(WINHTTP_OPTION_URL)"), GetLastError()));
+				ElectraHTTPStreamWinHttp::LogError(Response->GetErrorMessage());
+				return false;
+			}
+			FWideStringView EffUrl(UrlBuffer.GetData(), FCStringWide::Strlen(UrlBuffer.GetData()));
+			Response->EffectiveURL = EffUrl;
+		}
+
+		// Pre-allocate the receive buffer unless this is a HEAD request for which we will not get any data.
+		if (!Verb.Equals(TEXT("HEAD")))
+		{
+			// If the response was compressed WinHttp hides the Content-Length and Content-Encoding response headers from us, presumably
+			// as to avoid us using the size of the compressed data to allocate buffers that would turn out to be too small for the
+			// decompressed result. That is ok, we do not need the content length. During chunked encoding it would not necessarily
+			// be there anyway.
+			// Check for Content-Length header
+			if (Response->ContentLength.Len())
+			{
+				int64 cl = -1;
+				LexFromString(cl, *Response->ContentLength);
+				if (cl >= 0)
+				{
+					Response->ResponseData.SetLengthFromResponseHeader(cl);
+				}
+			}
+			// Alternatively check Content-Range header
+			else if (Response->ContentRange.Len())
+			{
+				ElectraHTTPStream::FHttpRange ContentRange;
+				if (ContentRange.ParseFromContentRangeResponse(Response->ContentRange))
+				{
+					if (ContentRange.GetNumberOfBytes() >= 0)
+					{
+						Response->ResponseData.SetLengthFromResponseHeader(ContentRange.GetNumberOfBytes());
+					}
+				}
+			}
+		}
+
+		Response->CurrentState = IElectraHTTPStreamResponse::EState::ReceivedResponseHeaders;
+		// Notify availability of response headers.
+		if (bNotify)
+		{
+			NotifyCallback(EElectraHTTPStreamNotificationReason::ReceivedHeaders, Response->ResponseHeaders.Num());
+		}
+	}
+	return true;
+}
+
+bool FElectraHTTPStreamRequestWinHttp::RequestResponseData()
+{
+	CurrentState = EState::AwaitingResponseData;
+	NumResponseBytesAvailableToRead = 0;
+	if (!WinHttpQueryDataAvailable(RequestHandle, NULL))
+	{
+		Response->SetErrorMessage(ElectraHTTPStreamWinHttp::GetErrorLogMessage(TEXT("WinHttpQueryDataAvailable()"), GetLastError()));
+		ElectraHTTPStreamWinHttp::LogError(Response->GetErrorMessage());
+		return false;
+	}
+	return true;
+}
+
+FElectraHTTPStreamRequestWinHttp::EReadResponseResult FElectraHTTPStreamRequestWinHttp::ReadResponseData()
+{
+	CurrentState = EState::ReadingResponseData;
+	Response->CurrentState = IElectraHTTPStreamResponse::EState::ReceivingResponseData;
+	if (NumResponseBytesAvailableToRead)
+	{
+		int32 rbSize = NumResponseBytesAvailableToRead > ELECTRA_HTTPSTREAM_INTERNAL_WINHTTP_READ_BUFFER_SIZE ? NumResponseBytesAvailableToRead : ELECTRA_HTTPSTREAM_INTERNAL_WINHTTP_READ_BUFFER_SIZE;
+		ChunkReadBuffer.Reserve(rbSize);
+		ChunkReadBuffer.AddUninitialized(NumResponseBytesAvailableToRead);
+		if (!WinHttpReadData(RequestHandle, ChunkReadBuffer.GetData(), NumResponseBytesAvailableToRead, NULL))
+		{
+			Response->SetErrorMessage(ElectraHTTPStreamWinHttp::GetErrorLogMessage(TEXT("WinHttpReadData()"), GetLastError()));
+			ElectraHTTPStreamWinHttp::LogError(Response->GetErrorMessage());
+			return EReadResponseResult::Failed;
+		}
+	}
+	else
+	{
+		return EReadResponseResult::EndOfData;
+	}
+	return EReadResponseResult::Success;
+}
+
+bool FElectraHTTPStreamRequestWinHttp::HandleNewResponseData()
+{
+	int64 ChunkBufferSize = ChunkReadBuffer.Num();
+	if (ChunkBufferSize)
+	{
+		// Add the data to the response using move semantics, clearing out the intermediate chunk read buffer.
+		Response->AddResponseData(MoveTemp(ChunkReadBuffer));
+		// Notify amount of new data available.
+		NotifyCallback(EElectraHTTPStreamNotificationReason::ReadData, ChunkBufferSize);
+	}
+	return true;
+}
+
+void FElectraHTTPStreamRequestWinHttp::Close(bool bHandlesOnly)
+{
+	if (RequestHandle)
+	{
+		// Clear the callback.
+		WinHttpSetStatusCallback(RequestHandle, NULL, WINHTTP_CALLBACK_FLAG_ALL_NOTIFICATIONS, 0);
+		// And close the handle.
+		WinHttpCloseHandle(RequestHandle);
+		RequestHandle = NULL;
+	}
+	if (ConnectionHandle)
+	{
+		WinHttpCloseHandle(ConnectionHandle);
+		ConnectionHandle = NULL;
+	}
+	if (!bHandlesOnly)
+	{
+		// Set EOS in the response receive buffer to signal no further data will arrive.
+		Response->SetEOS();
+	}
+}
+
+void FElectraHTTPStreamRequestWinHttp::Terminate()
+{
+	Close();
+	Response->SetErrorMessage(TEXT("Terminated due to HTTP module shutdown"));
+}
+
+void FElectraHTTPStreamRequestWinHttp::SetFinished()
+{
+	double Now = FPlatformTime::Seconds();
+	Response->TimeUntilFinished = Now - Response->StartTime;
+	Response->CurrentStatus = WasCanceled() ? IElectraHTTPStreamResponse::EStatus::Canceled : HasFailed() ?	IElectraHTTPStreamResponse::EStatus::Failed : IElectraHTTPStreamResponse::EStatus::Completed;
+	Response->CurrentState = IElectraHTTPStreamResponse::EState::Finished;
+	Response->SetEOS();
+	CurrentState = HasFailed() ? EState::Error : EState::Finished;
+}
+
+
+void FElectraHTTPStreamRequestWinHttp::WinHttpCallback(FElectraHTTPStreamWinHttp* Owner, HINTERNET hInternet, DWORD dwInternetStatus, void* lpvStatusInformation, DWORD dwStatusInformationLength)
+{
+	LLM_SCOPE(ELLMTag::MediaStreaming);
+
+	// If the handle was closed already, leave immediately.
+	if (!RequestHandle)
+	{
+		return;
+	}
+
+	double Now = FPlatformTime::Seconds();
+	const WINHTTP_ASYNC_RESULT* AsyncResult = nullptr;
+	switch(dwInternetStatus)
+	{
+		case WINHTTP_CALLBACK_STATUS_NAME_RESOLVED:
+		{
+			Response->TimeUntilNameResolved = Now - Response->StartTime;
+			CurrentState = CurrentState == EState::Error ? CurrentState : EState::Connecting;
+			break;
+		}
+		case WINHTTP_CALLBACK_STATUS_CONNECTED_TO_SERVER:
+		{
+			Response->TimeUntilConnected = Now - Response->StartTime;
+			CurrentState = CurrentState == EState::Error ? CurrentState : EState::SendingRequest;
+			break;
+		}
+		case WINHTTP_CALLBACK_STATUS_REQUEST_SENT:
+		{
+			Response->TimeUntilRequestSent = Now - Response->StartTime;
+			break;
+		}
+		case WINHTTP_CALLBACK_STATUS_SENDREQUEST_COMPLETE:
+		{
+			// Sending the request has completed. We can now get the response.
+			CurrentState = CurrentState == EState::Error ? CurrentState : EState::RequestSent;
+			Owner->TriggerWorkSignal();
+			break;
+		}
+		case WINHTTP_CALLBACK_STATUS_RESPONSE_RECEIVED:
+		{
+			CurrentState = CurrentState == EState::Error ? CurrentState : EState::ResponseReceived;
+			break;
+		}
+		case WINHTTP_CALLBACK_STATUS_INTERMEDIATE_RESPONSE:
+		{
+			// Intermediate response, like a 100-continue.
+			break;
+		}
+		case WINHTTP_CALLBACK_STATUS_REDIRECT:
+		{
+			if (--MaxRedirections < 0)
+			{
+				ParseHeaders(false);
+				Close(true);
+				CurrentState = EState::Error;
+				Response->SetErrorMessage(FString::Printf(TEXT("Exceeded maximum number of permitted redirections")));
+				ElectraHTTPStreamWinHttp::LogError(Response->GetErrorMessage());
+				Owner->TriggerWorkSignal();
+			}
+			else
+			{
+				// If we are doing a POST and get a 301 or 302 redirection we need to abort and issue a new request.
+				// This is because WinHttp turns the POST into a GET which is not good.
+				if (Verb.Equals(TEXT("POST")))
+				{
+					DWORD StatusCode = 0;
+					DWORD RequiredHeaderSize = sizeof(StatusCode);
+					WinHttpQueryHeaders(RequestHandle, WINHTTP_QUERY_STATUS_CODE | WINHTTP_QUERY_FLAG_NUMBER, WINHTTP_HEADER_NAME_BY_INDEX, &StatusCode, &RequiredHeaderSize, WINHTTP_NO_HEADER_INDEX);
+					if (StatusCode == 301 || StatusCode == 302)
+					{
+						RedirectedLocation = WCHAR_TO_TCHAR((LPWSTR)lpvStatusInformation);
+						Close(true);
+						bMustRedirect = true;
+						CurrentState = EState::NeedsRedirect;
+						Owner->TriggerWorkSignal();
+					}
+				}
+			}
+			break;
+		}
+		case WINHTTP_CALLBACK_STATUS_HEADERS_AVAILABLE:
+		{
+			Response->TimeUntilHeadersAvailable = Now - Response->StartTime;
+			CurrentState = CurrentState == EState::Error ? CurrentState : EState::HeadersAvailable;
+			Owner->TriggerWorkSignal();
+			break;
+		}
+		case WINHTTP_CALLBACK_STATUS_DATA_AVAILABLE:
+		{
+			if (Response->TimeUntilFirstByte == 0.0)
+			{
+				Response->TimeUntilFirstByte = Now - Response->StartTime;
+			}
+			check(dwStatusInformationLength == sizeof(DWORD));
+			if (dwStatusInformationLength == sizeof(DWORD))
+			{
+				NumResponseBytesAvailableToRead = *static_cast<DWORD*>(lpvStatusInformation);
+				CurrentState = CurrentState == EState::Error ? CurrentState : EState::ResponseDataAvailable;
+			}
+			else
+			{
+				CurrentState = EState::Error;
+				Response->SetErrorMessage(ElectraHTTPStreamWinHttp::GetErrorLogMessage(TEXT("WinHttpCallback(WINHTTP_CALLBACK_STATUS_DATA_AVAILABLE) not DWORD!"), ERROR_WINHTTP_INTERNAL_ERROR));
+				ElectraHTTPStreamWinHttp::LogError(Response->GetErrorMessage());
+			}
+			Owner->TriggerWorkSignal();
+			break;
+		}
+		case WINHTTP_CALLBACK_STATUS_READ_COMPLETE:
+		{
+			check(dwStatusInformationLength <= NumResponseBytesAvailableToRead);
+			Response->TimeOfMostRecentReceive = Now;
+			ChunkReadBuffer.SetNum(dwStatusInformationLength);
+			CurrentState = CurrentState == EState::Error ? CurrentState : EState::ResponseDataRead;
+			Owner->TriggerWorkSignal();
+			break;
+		}
+		case WINHTTP_CALLBACK_STATUS_WRITE_COMPLETE:
+		{
+			// We are currently leaving the sending of POST data to the initial WinHttpSendRequest() call
+			// and do not have to deal with this here.
+			break;
+		}
+		case WINHTTP_CALLBACK_STATUS_SECURE_FAILURE:
+		{
+			SecureFailureFlags = ElectraHTTPStreamWinHttp::GetSecurityErrorMessage(lpvStatusInformation ? (uint32)*static_cast<DWORD*>(lpvStatusInformation) : 0);
+			break;
+		}
+		case WINHTTP_CALLBACK_STATUS_REQUEST_ERROR:
+		{
+			CurrentState = EState::Error;
+			AsyncResult = static_cast<const WINHTTP_ASYNC_RESULT *>(lpvStatusInformation);
+			if (SecureFailureFlags.Len())
+			{
+				Response->SetErrorMessage(FString::Printf(TEXT("WinHttpCallback(WINHTTP_CALLBACK_STATUS_SECURE_FAILURE) \"%s\" during %s"), *SecureFailureFlags, *ElectraHTTPStreamWinHttp::GetApiName((uint32)AsyncResult->dwResult)));
+			}
+			else
+			{
+				Response->SetErrorMessage(FString::Printf(TEXT("WinHttpCallback(WINHTTP_CALLBACK_STATUS_REQUEST_ERROR) %s during %s"), *ElectraHTTPStreamWinHttp::GetErrorMessage(AsyncResult->dwError), *ElectraHTTPStreamWinHttp::GetApiName((uint32)AsyncResult->dwResult)));
+			}
+			ElectraHTTPStreamWinHttp::LogError(Response->GetErrorMessage());
+			Owner->TriggerWorkSignal();
+			break;
+		}
+		// Server-side connection close
+		case WINHTTP_CALLBACK_STATUS_CLOSING_CONNECTION:
+		case WINHTTP_CALLBACK_STATUS_CONNECTION_CLOSED:
+		{
+			break;
+		}
+		default:
+		{
+			break;
+		}
+	}
+
+	// Remember the last good callback reason to help locate the state before the problem occurred.
+	if (dwInternetStatus != WINHTTP_CALLBACK_STATUS_REQUEST_ERROR || LastWinHttpCallbackReason == 0)
+	{
+		LastWinHttpCallbackReason = dwInternetStatus;
+	}
+}
+
+
+/***************************************************************************************************************************************************/
+/***************************************************************************************************************************************************/
+/***************************************************************************************************************************************************/
+
+FElectraHTTPStreamWinHttp::FElectraHTTPStreamWinHttp()
+{
+}
+
+FElectraHTTPStreamWinHttp::~FElectraHTTPStreamWinHttp()
+{
+	Close();
+}
+
+bool FElectraHTTPStreamWinHttp::Initialize(const Electra::FParamDict& InOptions)
+{
+	LLM_SCOPE(ELLMTag::MediaStreaming);
+	// Create the WinHTTP session handle.
+	DWORD SessionFlags = WINHTTP_FLAG_ASYNC;
+
+#if ELECTRA_HTTPSTREAM_REQUIRES_SECURE_CONNECTIONS
+	SessionFlags |= WINHTTP_FLAG_SECURE_DEFAULTS;
+#endif
+
+	HINTERNET sh;
+	if (InOptions.HaveKey(TEXT("proxy")))
+	{
+		FString ProxyNameAndPort = InOptions.GetValue(TEXT("proxy")).SafeGetFString();
+		if (ProxyNameAndPort.Len())
+		{
+			sh = WinHttpOpen(NULL, WINHTTP_ACCESS_TYPE_NAMED_PROXY, TCHAR_TO_WCHAR(*ProxyNameAndPort), WINHTTP_NO_PROXY_BYPASS, SessionFlags);
+		}
+		else
+		{
+			sh = WinHttpOpen(NULL, WINHTTP_ACCESS_TYPE_AUTOMATIC_PROXY, WINHTTP_NO_PROXY_NAME, WINHTTP_NO_PROXY_BYPASS, SessionFlags);
+		}
+	}
+	else
+	{
+		sh = WinHttpOpen(NULL, WINHTTP_ACCESS_TYPE_AUTOMATIC_PROXY, WINHTTP_NO_PROXY_NAME, WINHTTP_NO_PROXY_BYPASS, SessionFlags);
+	}
+	if (sh == NULL)
+	{
+		ElectraHTTPStreamWinHttp::LogError(TEXT("WinHttpOpen()"), GetLastError());
+		return false;
+	}
+	SessionHandle = sh;
+
+	// Some settings can only be applied to the session handle, so we set those now.
+#if ELECTRA_HTTPSTREAM_REQUIRES_SECURE_CONNECTIONS
+	// Secure connections require TLS 1.2 or newer and will result in an error trying to
+	// set anything less than that, so we skip setting the protocols.
+#else
+	DWORD SecureProtocols = 0;
+	#ifdef WINHTTP_FLAG_SECURE_PROTOCOL_TLS1_1
+		SecureProtocols |= WINHTTP_FLAG_SECURE_PROTOCOL_TLS1_1;
+	#endif
+	#ifdef WINHTTP_FLAG_SECURE_PROTOCOL_TLS1_2
+		SecureProtocols |= WINHTTP_FLAG_SECURE_PROTOCOL_TLS1_2;
+	#endif
+	#ifdef WINHTTP_FLAG_SECURE_PROTOCOL_TLS1_3
+		SecureProtocols |= WINHTTP_FLAG_SECURE_PROTOCOL_TLS1_3;
+	#endif
+	if (SecureProtocols == 0) //-V547
+	{
+		SecureProtocols = WINHTTP_FLAG_SECURE_PROTOCOL_ALL & ~(WINHTTP_FLAG_SECURE_PROTOCOL_SSL2 | WINHTTP_FLAG_SECURE_PROTOCOL_SSL3);
+	}
+	if (!WinHttpSetOption(SessionHandle, WINHTTP_OPTION_SECURE_PROTOCOLS, &SecureProtocols, sizeof(SecureProtocols)))
+	{
+		ElectraHTTPStreamWinHttp::LogError(TEXT("WinHttpSetOption to set secure protocols on session"), GetLastError());
+		return false;
+	}
+#endif
+
+	// If support for HTTP/2 is available we enable it. This doesn't mean it gets used.
+	#ifdef ELECTRA_HTTPSTREAM_INTERNAL_WINHTTP_ENABLE_HTTP2
+	DWORD FlagEnableHTTP2 = WINHTTP_PROTOCOL_FLAG_HTTP2;
+	if (!WinHttpSetOption(SessionHandle, WINHTTP_OPTION_ENABLE_HTTP_PROTOCOL, &FlagEnableHTTP2, sizeof(FlagEnableHTTP2)))
+	{
+		UE_LOG(LogElectraHTTPStream, Verbose, TEXT("HTTP/2 is not supported, falling back to HTTP/1.1"));
+	}
+	#endif
+
+	// Create the worker thread.
+	Thread = FRunnableThread::Create(this, TEXT("ElectraHTTPStream"), 128 * 1024, TPri_Normal);
+	if (!Thread)
+	{
+		UE_LOG(LogElectraHTTPStream, Error, TEXT("Failed to create the ElectraHTTPStream worker thread"));
+		return false;
+	}
+	return true;
+}
+
+void FElectraHTTPStreamWinHttp::Close()
+{
+	LLM_SCOPE(ELLMTag::MediaStreaming);
+	if (Thread)
+	{
+		Thread->Kill(true);
+		delete Thread;
+		Thread = nullptr;
+	}
+
+	if (SessionHandle)
+	{
+		WinHttpCloseHandle(SessionHandle);
+		SessionHandle = NULL;
+	}
+
+	RemoveOutdatedRequests();
+}
+
+IElectraHTTPStreamRequestPtr FElectraHTTPStreamWinHttp::CreateRequest()
+{
+	LLM_SCOPE(ELLMTag::MediaStreaming);
+	RemoveOutdatedRequests();
+	return MakeShared<FElectraHTTPStreamRequestWinHttp, ESPMode::ThreadSafe>(GetNewRequestIndex());
+}
+
+void FElectraHTTPStreamWinHttp::AddRequest(IElectraHTTPStreamRequestPtr InRequest)
+{
+	if (InRequest.IsValid())
+	{
+		if (Thread)
+		{
+			FScopeLock lock(&RequestLock);
+			NewRequests.Emplace(StaticCastSharedPtr<FElectraHTTPStreamRequestWinHttp>(InRequest));
+			TriggerWorkSignal();
+		}
+		else
+		{
+			TSharedPtr<FElectraHTTPStreamRequestWinHttp, ESPMode::ThreadSafe> Req = StaticCastSharedPtr<FElectraHTTPStreamRequestWinHttp>(InRequest);
+			Req->Terminate();
+			Req->NotifyCallback(EElectraHTTPStreamNotificationReason::Completed, 1);
+		}
+	}
+}
+
+void FElectraHTTPStreamWinHttp::Stop()
+{
+	ExitRequest.Set(1);
+}
+
+uint32 FElectraHTTPStreamWinHttp::Run()
+{
+	LLM_SCOPE(ELLMTag::MediaStreaming);
+
+	while(!ExitRequest.GetValue())
+	{
+		HaveWorkSignal.WaitTimeoutAndReset(20);
+
+		{
+		SCOPE_CYCLE_COUNTER(STAT_ElectraHTTPThread_Process);
+		SetupNewRequests();
+		UpdateActiveRequests();
+		HandleCompletedRequests();
+		}
+
+		// User callback
+		{
+		SCOPE_CYCLE_COUNTER(STAT_ElectraHTTPThread_CustomHandler);
+		ThreadHandlerCallback.ExecuteIfBound();
+		}
+	}
+
+	// Remove requests.
+	RequestLock.Lock();
+	while(NewRequests.Num())
+	{
+		TSharedPtr<FElectraHTTPStreamRequestWinHttp, ESPMode::ThreadSafe> Req = NewRequests.Pop();
+		Req->Terminate();
+		CompletedRequests.Emplace(MoveTemp(Req));
+	}
+	while(ActiveRequests.Num())
+	{
+		TSharedPtr<FElectraHTTPStreamRequestWinHttp, ESPMode::ThreadSafe> Req = ActiveRequests.Pop();
+		Req->Terminate();
+		CompletedRequests.Emplace(MoveTemp(Req));
+	}
+	RequestLock.Unlock();
+	HandleCompletedRequests();
+
+	return 0;
+}
+
+void FElectraHTTPStreamWinHttp::_WinHttpCallback(HINTERNET hInternet, DWORD_PTR dwContext, DWORD dwInternetStatus, void* lpvStatusInformation, DWORD dwStatusInformationLength)
+{
+	SCOPE_CYCLE_COUNTER(STAT_ElectraHTTPThread_AsyncCallback);
+
+	FRequestPointers Pointers = GetRequestByIndex(static_cast<uint64>(dwContext));
+	TSharedPtr<FElectraHTTPStreamWinHttp, ESPMode::ThreadSafe> Self = Pointers.Owner.Pin();
+	TSharedPtr<FElectraHTTPStreamRequestWinHttp, ESPMode::ThreadSafe> Request = Pointers.Request.Pin();
+	if (Self.IsValid() && Request.IsValid())
+	{
+		static_cast<FElectraHTTPStreamRequestWinHttp*>(Request.Get())->WinHttpCallback(Self.Get(), hInternet, dwInternetStatus, lpvStatusInformation, dwStatusInformationLength);
+	}
+}
+
+void FElectraHTTPStreamWinHttp::SetupNewRequests()
+{
+	RequestLock.Lock();
+	TArray<TSharedPtr<FElectraHTTPStreamRequestWinHttp, ESPMode::ThreadSafe>> NewReqs;
+	Swap(NewReqs, NewRequests);
+	RequestLock.Unlock();
+	for(auto &Request : NewReqs)
+	{
+		if (Request->Setup(this))
+		{
+			FRequestPointers Pointers;
+			Pointers.Owner = AsShared();
+			Pointers.Request = Request;
+			AllRequestHandleLock.Lock();
+			AllRequestHandles.Add(Request->GetRequestIndex(), MoveTemp(Pointers));
+			AllRequestHandleLock.Unlock();
+			ActiveRequests.Emplace(Request);
+			if (Request->WasCanceled() || !Request->Execute())
+			{
+				ActiveRequests.Remove(Request);
+				CompletedRequests.Emplace(Request);
+			}
+		}
+		else
+		{
+			CompletedRequests.Emplace(Request);
+		}
+	}
+}
+
+void FElectraHTTPStreamWinHttp::UpdateActiveRequests()
+{
+	for(int32 i=0; i<ActiveRequests.Num(); ++i)
+	{
+		TSharedPtr<FElectraHTTPStreamRequestWinHttp, ESPMode::ThreadSafe> Request = ActiveRequests[i];
+		bool bRemoveRequest = false;
+		bool bErrored = false;
+
+		// If the request got canceled do not handle it any further and move it to the completed stage.
+		if (Request->WasCanceled())
+		{
+			bRemoveRequest = true;
+		}
+		// If the request has failed in any way, do not handle it any further and move it to the completed stage.
+		else if (Request->HasFailed())
+		{
+			bErrored = true;
+		}
+		else
+		{
+			switch(Request->GetCurrentState())
+			{
+				case FElectraHTTPStreamRequestWinHttp::EState::NeedsRedirect:
+				{
+					// Remove this request from the map immediately.
+					AllRequestHandleLock.Lock();
+					AllRequestHandles.Remove(Request->GetRequestIndex());
+					AllRequestHandleLock.Unlock();
+					// Also remove it from the active requests.
+					ActiveRequests.RemoveAt(i);
+					--i;
+					// And add it again as a new request.
+					RequestLock.Lock();
+					NewRequests.Emplace(Request);
+					RequestLock.Unlock();
+					break;
+				}
+				case FElectraHTTPStreamRequestWinHttp::EState::RequestSent:
+				{
+					bErrored = !Request->RequestResponse();
+					break;
+				}
+				case FElectraHTTPStreamRequestWinHttp::EState::HeadersAvailable:
+				{
+					bErrored = !Request->ParseHeaders(true);
+					if (!bErrored)
+					{
+						bErrored = !Request->RequestResponseData();
+					}
+					break;
+				}
+				case FElectraHTTPStreamRequestWinHttp::EState::ResponseDataAvailable:
+				{
+					FElectraHTTPStreamRequestWinHttp::EReadResponseResult Result = Request->ReadResponseData();
+					if (Result == FElectraHTTPStreamRequestWinHttp::EReadResponseResult::Failed)
+					{
+						bErrored = true;
+					}
+					else if (Result == FElectraHTTPStreamRequestWinHttp::EReadResponseResult::EndOfData)
+					{
+						bRemoveRequest = true;
+					}
+					break;
+				}
+				case FElectraHTTPStreamRequestWinHttp::EState::ResponseDataRead:
+				{
+					bErrored = !Request->HandleNewResponseData();
+					if (!bErrored)
+					{
+						bErrored = !Request->RequestResponseData();
+					}
+					break;
+				}
+				default:
+				{
+					break;
+				}
+			}
+		}
+		if (bErrored)
+		{
+			bRemoveRequest = true;
+		}
+		if (bRemoveRequest)
+		{
+			ActiveRequests.RemoveAt(i);
+			--i;
+			CompletedRequests.Emplace(MoveTemp(Request));
+		}
+	}
+}
+
+void FElectraHTTPStreamWinHttp::HandleCompletedRequests()
+{
+	if (CompletedRequests.Num())
+	{
+		TArray<TSharedPtr<FElectraHTTPStreamRequestWinHttp, ESPMode::ThreadSafe>> TempRequests;
+		Swap(CompletedRequests, TempRequests);
+		for(auto &Finished : TempRequests)
+		{
+			Finished->SetFinished();
+			Finished->Close();
+			AllRequestHandleLock.Lock();
+			AllRequestHandles.Remove(Finished->GetRequestIndex());
+			AllRequestHandleLock.Unlock();
+			// Parameter is 0 when finished successfully or 1 otherwise
+			int64 Param = Finished->HasFailed() ? 1 : 0;
+			if (!Finished->WasCanceled())
+			{
+				Finished->NotifyCallback(EElectraHTTPStreamNotificationReason::Completed, Param);
+			}
+		}
+	}
+}
+
+
+/***************************************************************************************************************************************************/
+/***************************************************************************************************************************************************/
+/***************************************************************************************************************************************************/
+
+void FPlatformElectraHTTPStreamWinHttp::Startup()
+{
+}
+
+void FPlatformElectraHTTPStreamWinHttp::Shutdown()
+{
+}
+
+TSharedPtr<IElectraHTTPStream, ESPMode::ThreadSafe> FPlatformElectraHTTPStreamWinHttp::Create(const Electra::FParamDict& InOptions)
+{
+	TSharedPtr<FElectraHTTPStreamWinHttp, ESPMode::ThreadSafe> New = MakeShareable(new FElectraHTTPStreamWinHttp);
+	if (New.IsValid())
+	{
+		if (!New->Initialize(InOptions))
+		{
+			New.Reset();
+		}
+	}
+	return New;
+}
+
+#endif // ELECTRA_HTTPSTREAM_WINHTTP
diff --git a/Engine/Plugins/Media/ElectraUtil/Source/ElectraHTTPStream/Private/WinHttp/WinHttpElectraHTTPStream.h b/Engine/Plugins/Media/ElectraUtil/Source/ElectraHTTPStream/Private/WinHttp/WinHttpElectraHTTPStream.h
new file mode 100644
--- /dev/null
+++ b/Engine/Plugins/Media/ElectraUtil/Source/ElectraHTTPStream/Private/WinHttp/WinHttpElectraHTTPStream.h
@@ -0,0 +1,21 @@
+// Copyright Epic Games, Inc. All Rights Reserved.
+
+#pragma once
+
+#if ELECTRA_HTTPSTREAM_WINHTTP
+
+#include "CoreMinimal.h"
+#include "Templates/SharedPointer.h"
+#include "ParameterDictionary.h"
+
+class ELECTRAHTTPSTREAM_API FPlatformElectraHTTPStreamWinHttp
+{
+public:
+	static void Startup();
+	static void Shutdown();
+	static TSharedPtr<IElectraHTTPStream, ESPMode::ThreadSafe> Create(const Electra::FParamDict& InOptions);
+};
+
+typedef FPlatformElectraHTTPStreamWinHttp FPlatformElectraHTTPStream;
+
+#endif // ELECTRA_HTTPSTREAM_WINHTTP
diff --git a/Engine/Plugins/Media/ElectraUtil/Source/ElectraHTTPStream/Private/WinHttp/WinHttpLogError.cpp b/Engine/Plugins/Media/ElectraUtil/Source/ElectraHTTPStream/Private/WinHttp/WinHttpLogError.cpp
new file mode 100644
--- /dev/null
+++ b/Engine/Plugins/Media/ElectraUtil/Source/ElectraHTTPStream/Private/WinHttp/WinHttpLogError.cpp
@@ -0,0 +1,170 @@
+// Copyright Epic Games, Inc. All Rights Reserved.
+
+#if ELECTRA_HTTPSTREAM_WINHTTP
+
+#include "WinHttp/WinHttpLogError.h"
+#include "ElectraHTTPStreamModule.h"
+
+#include "Windows/AllowWindowsPlatformTypes.h"
+#include <errhandlingapi.h>
+#include <winhttp.h>
+#include "Windows/HideWindowsPlatformTypes.h"
+
+
+namespace ElectraHTTPStreamWinHttp
+{
+
+FString GetSecurityErrorMessage(uint32 Flags)
+{
+	auto PrependDelimiter = [](FString& s) -> void
+	{
+		if (s.Len())
+		{
+			s.InsertAt(0, TEXT("; "));
+		}
+	};
+
+	FString msg;
+	if ((Flags & WINHTTP_CALLBACK_STATUS_FLAG_CERT_REV_FAILED) != 0)
+	{
+		msg.Append(TEXT("Certification revocation check failed"));
+	}
+	if ((Flags & WINHTTP_CALLBACK_STATUS_FLAG_INVALID_CERT) != 0)
+	{
+		PrependDelimiter(msg);
+		msg.Append(TEXT("SSL certificate is invalid"));
+	}
+	if ((Flags & WINHTTP_CALLBACK_STATUS_FLAG_CERT_REVOKED) != 0)
+	{
+		PrependDelimiter(msg);
+		msg.Append(TEXT("SSL certificate was revoked"));
+	}
+	if ((Flags & WINHTTP_CALLBACK_STATUS_FLAG_INVALID_CA) != 0)
+	{
+		PrependDelimiter(msg);
+		msg.Append(TEXT("Unfamiliar Certificate Authority that generated the server's certificate"));
+	}
+	if ((Flags & WINHTTP_CALLBACK_STATUS_FLAG_CERT_CN_INVALID) != 0)
+	{
+		PrependDelimiter(msg);
+		msg.Append(TEXT("SSL certificate common name (host name field) is incorrect"));
+	}
+	if ((Flags & WINHTTP_CALLBACK_STATUS_FLAG_CERT_DATE_INVALID) != 0)
+	{
+		PrependDelimiter(msg);
+		msg.Append(TEXT("SSL certificate date is expired"));
+	}
+	if ((Flags & WINHTTP_CALLBACK_STATUS_FLAG_CERT_WRONG_USAGE) != 0)
+	{
+		PrependDelimiter(msg);
+		msg.Append(TEXT("Wrong SSL certificate usage"));
+	}
+	if ((Flags & WINHTTP_CALLBACK_STATUS_FLAG_SECURITY_CHANNEL_ERROR) != 0)
+	{
+		PrependDelimiter(msg);
+		msg.Append(TEXT("The application experienced an internal error loading the SSL libraries"));
+	}
+	return msg;
+}
+
+FString GetApiName(uint32 ApiNameCode)
+{
+	switch(ApiNameCode)
+	{
+		case API_RECEIVE_RESPONSE:		return TEXT("RECEIVE_RESPONSE");
+		case API_QUERY_DATA_AVAILABLE:	return TEXT("QUERY_DATA_AVAILABLE");
+		case API_READ_DATA:				return TEXT("READ_DATA");
+		case API_WRITE_DATA:			return TEXT("WRITE_DATA");
+		case API_SEND_REQUEST:			return TEXT("SEND_REQUEST");
+		case API_GET_PROXY_FOR_URL:		return TEXT("GET_PROXY_FOR_URL");
+		default:						return FString::Printf(TEXT("Unknown API name code %u"), ApiNameCode);
+	}
+}
+
+
+FString GetErrorMessage(uint32 ErrorCode)
+{
+	switch(ErrorCode)
+	{
+		case ERROR_WINHTTP_OUT_OF_HANDLES:							return TEXT("OUT_OF_HANDLES");
+		case ERROR_WINHTTP_TIMEOUT:									return TEXT("TIMEOUT");
+		case ERROR_WINHTTP_INTERNAL_ERROR:							return TEXT("INTERNAL_ERROR");
+		case ERROR_WINHTTP_INVALID_URL:								return TEXT("INVALID_URL");
+		case ERROR_WINHTTP_UNRECOGNIZED_SCHEME:						return TEXT("UNRECOGNIZED_SCHEME");
+		case ERROR_WINHTTP_NAME_NOT_RESOLVED:						return TEXT("NAME_NOT_RESOLVED");
+		case ERROR_WINHTTP_INVALID_OPTION:							return TEXT("INVALID_OPTION");
+		case ERROR_WINHTTP_OPTION_NOT_SETTABLE:						return TEXT("OPTION_NOT_SETTABLE");
+		case ERROR_WINHTTP_SHUTDOWN:								return TEXT("SHUTDOWN");
+		case ERROR_WINHTTP_LOGIN_FAILURE:							return TEXT("LOGIN_FAILURE");
+		case ERROR_WINHTTP_OPERATION_CANCELLED:						return TEXT("OPERATION_CANCELLED");
+		case ERROR_WINHTTP_INCORRECT_HANDLE_TYPE:					return TEXT("INCORRECT_HANDLE_TYPE");
+		case ERROR_WINHTTP_INCORRECT_HANDLE_STATE:					return TEXT("INCORRECT_HANDLE_STATE");
+		case ERROR_WINHTTP_CANNOT_CONNECT:							return TEXT("CANNOT_CONNECT");
+		case ERROR_WINHTTP_CONNECTION_ERROR:						return TEXT("CONNECTION_ERROR");
+		case ERROR_WINHTTP_RESEND_REQUEST:							return TEXT("RESEND_REQUEST");
+		case ERROR_WINHTTP_CLIENT_AUTH_CERT_NEEDED:					return TEXT("CLIENT_AUTH_CERT_NEEDED");
+		case ERROR_WINHTTP_CANNOT_CALL_BEFORE_OPEN:					return TEXT("CANNOT_CALL_BEFORE_OPEN");
+		case ERROR_WINHTTP_CANNOT_CALL_BEFORE_SEND:					return TEXT("CANNOT_CALL_BEFORE_SEND");
+		case ERROR_WINHTTP_CANNOT_CALL_AFTER_SEND:					return TEXT("CANNOT_CALL_AFTER_SEND");
+		case ERROR_WINHTTP_CANNOT_CALL_AFTER_OPEN:					return TEXT("CANNOT_CALL_AFTER_OPEN");
+		case ERROR_WINHTTP_HEADER_NOT_FOUND:						return TEXT("HEADER_NOT_FOUND");
+		case ERROR_WINHTTP_INVALID_SERVER_RESPONSE:					return TEXT("INVALID_SERVER_RESPONSE");
+		case ERROR_WINHTTP_INVALID_HEADER:							return TEXT("INVALID_HEADER");
+		case ERROR_WINHTTP_INVALID_QUERY_REQUEST:					return TEXT("INVALID_QUERY_REQUEST");
+		case ERROR_WINHTTP_HEADER_ALREADY_EXISTS:					return TEXT("HEADER_ALREADY_EXISTS");
+		case ERROR_WINHTTP_REDIRECT_FAILED:							return TEXT("REDIRECT_FAILED");
+		case ERROR_WINHTTP_AUTO_PROXY_SERVICE_ERROR:				return TEXT("AUTO_PROXY_SERVICE_ERROR");
+		case ERROR_WINHTTP_BAD_AUTO_PROXY_SCRIPT:					return TEXT("BAD_AUTO_PROXY_SCRIPT");
+		case ERROR_WINHTTP_UNABLE_TO_DOWNLOAD_SCRIPT:				return TEXT("UNABLE_TO_DOWNLOAD_SCRIPT");
+		case ERROR_WINHTTP_UNHANDLED_SCRIPT_TYPE:					return TEXT("UNHANDLED_SCRIPT_TYPE");
+		case ERROR_WINHTTP_SCRIPT_EXECUTION_ERROR:					return TEXT("SCRIPT_EXECUTION_ERROR");
+		case ERROR_WINHTTP_NOT_INITIALIZED:							return TEXT("NOT_INITIALIZED");
+		case ERROR_WINHTTP_SECURE_FAILURE:							return TEXT("SECURE_FAILURE");
+		case ERROR_WINHTTP_SECURE_CERT_DATE_INVALID:				return TEXT("SECURE_CERT_DATE_INVALID");
+		case ERROR_WINHTTP_SECURE_CERT_CN_INVALID:					return TEXT("SECURE_CERT_CN_INVALID");
+		case ERROR_WINHTTP_SECURE_INVALID_CA:						return TEXT("SECURE_INVALID_CA");
+		case ERROR_WINHTTP_SECURE_CERT_REV_FAILED:					return TEXT("SECURE_CERT_REV_FAILED");
+		case ERROR_WINHTTP_SECURE_CHANNEL_ERROR:					return TEXT("SECURE_CHANNEL_ERROR");
+		case ERROR_WINHTTP_SECURE_INVALID_CERT:						return TEXT("SECURE_INVALID_CERT");
+		case ERROR_WINHTTP_SECURE_CERT_REVOKED:						return TEXT("SECURE_CERT_REVOKED");
+		case ERROR_WINHTTP_SECURE_CERT_WRONG_USAGE:					return TEXT("SECURE_CERT_WRONG_USAGE");
+		case ERROR_WINHTTP_AUTODETECTION_FAILED:					return TEXT("AUTODETECTION_FAILED");
+		case ERROR_WINHTTP_HEADER_COUNT_EXCEEDED:					return TEXT("HEADER_COUNT_EXCEEDED");
+		case ERROR_WINHTTP_HEADER_SIZE_OVERFLOW:					return TEXT("HEADER_SIZE_OVERFLOW");
+		case ERROR_WINHTTP_CHUNKED_ENCODING_HEADER_SIZE_OVERFLOW:	return TEXT("CHUNKED_ENCODING_HEADER_SIZE_OVERFLOW");
+		case ERROR_WINHTTP_RESPONSE_DRAIN_OVERFLOW:					return TEXT("RESPONSE_DRAIN_OVERFLOW");
+		case ERROR_WINHTTP_CLIENT_CERT_NO_PRIVATE_KEY:				return TEXT("CLIENT_CERT_NO_PRIVATE_KEY");
+		case ERROR_WINHTTP_CLIENT_CERT_NO_ACCESS_PRIVATE_KEY:		return TEXT("CLIENT_CERT_NO_ACCESS_PRIVATE_KEY");
+		case ERROR_WINHTTP_CLIENT_AUTH_CERT_NEEDED_PROXY:			return TEXT("CLIENT_AUTH_CERT_NEEDED_PROXY");
+		case ERROR_WINHTTP_SECURE_FAILURE_PROXY:					return TEXT("SECURE_FAILURE_PROXY");
+		case ERROR_WINHTTP_HTTP_PROTOCOL_MISMATCH:					return TEXT("HTTP_PROTOCOL_MISMATCH");
+		case 5:														return TEXT("ACCESS_DENIED");		// When an unsafe operation is attempted on a secure handle
+		default:													return FString::Printf(TEXT("Unknown error code 0x%08x"), ErrorCode);
+	}
+}
+
+FString GetErrorLogMessage(const TCHAR* const Method, uint32 ErrorCode)
+{
+	if (Method)
+	{
+		return FString::Printf(TEXT("Error calling WinHttp method %s: %s"), Method, *GetErrorMessage(ErrorCode));
+	}
+	else
+	{
+		return FString::Printf(TEXT("Error calling WinHttp method: %s"), *GetErrorMessage(ErrorCode));
+	}
+}
+
+void LogError(const FString& Message)
+{
+	UE_LOG(LogElectraHTTPStream, Error, TEXT("%s"), *Message);
+}
+
+void LogError(const TCHAR* const Method, uint32 ErrorCode)
+{
+	UE_LOG(LogElectraHTTPStream, Error, TEXT("%s"), *GetErrorLogMessage(Method, ErrorCode));
+}
+
+}
+
+#endif // ELECTRA_HTTPSTREAM_WINHTTP
diff --git a/Engine/Plugins/Media/ElectraUtil/Source/ElectraHTTPStream/Private/WinHttp/WinHttpLogError.h b/Engine/Plugins/Media/ElectraUtil/Source/ElectraHTTPStream/Private/WinHttp/WinHttpLogError.h
new file mode 100644
--- /dev/null
+++ b/Engine/Plugins/Media/ElectraUtil/Source/ElectraHTTPStream/Private/WinHttp/WinHttpLogError.h
@@ -0,0 +1,15 @@
+// Copyright Epic Games, Inc. All Rights Reserved.
+
+#pragma once
+
+#include "CoreMinimal.h"
+
+namespace ElectraHTTPStreamWinHttp
+{
+	FString GetSecurityErrorMessage(uint32 Flags);
+	FString GetApiName(uint32 ApiNameCode);
+	FString GetErrorMessage(uint32 ErrorCode);
+	FString GetErrorLogMessage(const TCHAR* const Method, uint32 ErrorCode);
+	void LogError(const TCHAR* const Method, uint32 ErrorCode);
+	void LogError(const FString& Message);
+}
diff --git a/Engine/Plugins/Media/ElectraUtil/Source/ElectraHTTPStream/Public/Android/AndroidPlatformElectraHTTPStream.h b/Engine/Plugins/Media/ElectraUtil/Source/ElectraHTTPStream/Public/Android/AndroidPlatformElectraHTTPStream.h
new file mode 100644
--- /dev/null
+++ b/Engine/Plugins/Media/ElectraUtil/Source/ElectraHTTPStream/Public/Android/AndroidPlatformElectraHTTPStream.h
@@ -0,0 +1,8 @@
+// Copyright Epic Games, Inc. All Rights Reserved.
+
+#pragma once
+
+#if ELECTRA_HTTPSTREAM_LIBCURL
+#include "Curl/CurlElectraHTTPStream.h"
+
+#endif
diff --git a/Engine/Plugins/Media/ElectraUtil/Source/ElectraHTTPStream/Public/IOS/IOSPlatformElectraHTTPStream.h b/Engine/Plugins/Media/ElectraUtil/Source/ElectraHTTPStream/Public/IOS/IOSPlatformElectraHTTPStream.h
new file mode 100644
--- /dev/null
+++ b/Engine/Plugins/Media/ElectraUtil/Source/ElectraHTTPStream/Public/IOS/IOSPlatformElectraHTTPStream.h
@@ -0,0 +1,8 @@
+// Copyright Epic Games, Inc. All Rights Reserved.
+
+#pragma once
+
+#if ELECTRA_HTTPSTREAM_APPLE
+#include "Apple/AppleElectraHTTPStream.h"
+
+#endif
diff --git a/Engine/Plugins/Media/ElectraUtil/Source/ElectraHTTPStream/Public/Linux/LinuxPlatformElectraHTTPStream.h b/Engine/Plugins/Media/ElectraUtil/Source/ElectraHTTPStream/Public/Linux/LinuxPlatformElectraHTTPStream.h
new file mode 100644
--- /dev/null
+++ b/Engine/Plugins/Media/ElectraUtil/Source/ElectraHTTPStream/Public/Linux/LinuxPlatformElectraHTTPStream.h
@@ -0,0 +1,8 @@
+// Copyright Epic Games, Inc. All Rights Reserved.
+
+#pragma once
+
+#if ELECTRA_HTTPSTREAM_LIBCURL
+#include "Curl/CurlElectraHTTPStream.h"
+
+#endif
diff --git a/Engine/Plugins/Media/ElectraUtil/Source/ElectraHTTPStream/Public/Mac/MacPlatformElectraHTTPStream.h b/Engine/Plugins/Media/ElectraUtil/Source/ElectraHTTPStream/Public/Mac/MacPlatformElectraHTTPStream.h
new file mode 100644
--- /dev/null
+++ b/Engine/Plugins/Media/ElectraUtil/Source/ElectraHTTPStream/Public/Mac/MacPlatformElectraHTTPStream.h
@@ -0,0 +1,8 @@
+// Copyright Epic Games, Inc. All Rights Reserved.
+
+#pragma once
+
+#if ELECTRA_HTTPSTREAM_APPLE
+#include "Apple/AppleElectraHTTPStream.h"
+
+#endif
diff --git a/Engine/Plugins/Media/ElectraUtil/Source/ElectraHTTPStream/Public/PlatformElectraHTTPStream.h b/Engine/Plugins/Media/ElectraUtil/Source/ElectraHTTPStream/Public/PlatformElectraHTTPStream.h
new file mode 100644
--- /dev/null
+++ b/Engine/Plugins/Media/ElectraUtil/Source/ElectraHTTPStream/Public/PlatformElectraHTTPStream.h
@@ -0,0 +1,6 @@
+// Copyright Epic Games, Inc. All Rights Reserved.
+#pragma once
+
+#include "CoreMinimal.h"
+
+#include COMPILED_PLATFORM_HEADER(PlatformElectraHTTPStream.h)
diff --git a/Engine/Plugins/Media/ElectraUtil/Source/ElectraHTTPStream/Public/Windows/WindowsPlatformElectraHTTPStream.h b/Engine/Plugins/Media/ElectraUtil/Source/ElectraHTTPStream/Public/Windows/WindowsPlatformElectraHTTPStream.h
new file mode 100644
--- /dev/null
+++ b/Engine/Plugins/Media/ElectraUtil/Source/ElectraHTTPStream/Public/Windows/WindowsPlatformElectraHTTPStream.h
@@ -0,0 +1,11 @@
+// Copyright Epic Games, Inc. All Rights Reserved.
+
+#pragma once
+
+#if ELECTRA_HTTPSTREAM_WINHTTP
+#include "WinHttp/WinHttpElectraHTTPStream.h"
+
+#elif ELECTRA_HTTPSTREAM_LIBCURL
+#include "Curl/CurlElectraHTTPStream.h"
+
+#endif
