ElectraHTTPStream: Fix non unity compilation errors

#rnx

diff --git a/Engine/Plugins/Media/ElectraUtil/Source/ElectraHTTPStream/Private/Curl/CurlElectraHTTPStream.cpp b/Engine/Plugins/Media/ElectraUtil/Source/ElectraHTTPStream/Private/Curl/CurlElectraHTTPStream.cpp
--- a/Engine/Plugins/Media/ElectraUtil/Source/ElectraHTTPStream/Private/Curl/CurlElectraHTTPStream.cpp
+++ b/Engine/Plugins/Media/ElectraUtil/Source/ElectraHTTPStream/Private/Curl/CurlElectraHTTPStream.cpp
@@ -1262,6 +1262,7 @@ void FElectraHTTPStreamLibCurl::HandleCompletedRequests()
 /***************************************************************************************************************************************************/
 /***************************************************************************************************************************************************/
 
+#include "Modules/ModuleManager.h"
 #include "HttpModule.h"
 
 void FPlatformElectraHTTPStreamLibCurl::Startup()
diff --git a/Engine/Plugins/Media/ElectraUtil/Source/ElectraHTTPStream/Private/Curl/CurlElectraHTTPStreamImpl.h b/Engine/Plugins/Media/ElectraUtil/Source/ElectraHTTPStream/Private/Curl/CurlElectraHTTPStreamImpl.h
--- a/Engine/Plugins/Media/ElectraUtil/Source/ElectraHTTPStream/Private/Curl/CurlElectraHTTPStreamImpl.h
+++ b/Engine/Plugins/Media/ElectraUtil/Source/ElectraHTTPStream/Private/Curl/CurlElectraHTTPStreamImpl.h
@@ -4,6 +4,13 @@
 
 #if ELECTRA_HTTPSTREAM_LIBCURL
 
+#include "CoreMinimal.h"
+#include "HAL/PlatformTime.h"
+#include "HAL/PlatformProcess.h"
+#include "HAL/Runnable.h"
+#include "HAL/RunnableThread.h"
+#include "Templates/SharedPointer.h"
+
 #include "Curl/CurlElectraHTTPStreamConfig.h"
 #include "Curl/CurlElectra.h"
 #include "Utilities/TimeWaitableSignal.h"
diff --git a/Engine/Plugins/Media/ElectraUtil/Source/ElectraHTTPStream/Private/ElectraHTTPStreamResponse.h b/Engine/Plugins/Media/ElectraUtil/Source/ElectraHTTPStream/Private/ElectraHTTPStreamResponse.h
--- a/Engine/Plugins/Media/ElectraUtil/Source/ElectraHTTPStream/Private/ElectraHTTPStreamResponse.h
+++ b/Engine/Plugins/Media/ElectraUtil/Source/ElectraHTTPStream/Private/ElectraHTTPStreamResponse.h
@@ -2,6 +2,13 @@
 
 #pragma once
 
+#include "CoreMinimal.h"
+#include "HAL/PlatformTime.h"
+#include "HAL/PlatformProcess.h"
+#include "HAL/Runnable.h"
+#include "HAL/RunnableThread.h"
+#include "Templates/SharedPointer.h"
+
 #include "ElectraHTTPStream.h"
 #include "ElectraHTTPStreamBuffer.h"
 
diff --git a/Engine/Plugins/Media/ElectraUtil/Source/ElectraHTTPStream/Private/WinHttp/WinHttpElectraHTTPStream.h b/Engine/Plugins/Media/ElectraUtil/Source/ElectraHTTPStream/Private/WinHttp/WinHttpElectraHTTPStream.h
--- a/Engine/Plugins/Media/ElectraUtil/Source/ElectraHTTPStream/Private/WinHttp/WinHttpElectraHTTPStream.h
+++ b/Engine/Plugins/Media/ElectraUtil/Source/ElectraHTTPStream/Private/WinHttp/WinHttpElectraHTTPStream.h
@@ -8,6 +8,8 @@
 #include "Templates/SharedPointer.h"
 #include "ParameterDictionary.h"
 
+class IElectraHTTPStream;
+
 class ELECTRAHTTPSTREAM_API FPlatformElectraHTTPStreamWinHttp
 {
 public:
