ElectraHTTPStream: disabling TLS1.3 on WinHttp since this is not available on all versions of Windows 10

#rnx

diff --git a/Engine/Plugins/Media/ElectraUtil/Source/ElectraHTTPStream/Private/WinHttp/WinHttpElectraHTTPStream.cpp b/Engine/Plugins/Media/ElectraUtil/Source/ElectraHTTPStream/Private/WinHttp/WinHttpElectraHTTPStream.cpp
--- a/Engine/Plugins/Media/ElectraUtil/Source/ElectraHTTPStream/Private/WinHttp/WinHttpElectraHTTPStream.cpp
+++ b/Engine/Plugins/Media/ElectraUtil/Source/ElectraHTTPStream/Private/WinHttp/WinHttpElectraHTTPStream.cpp
@@ -1092,9 +1092,13 @@ bool FElectraHTTPStreamWinHttp::Initialize(const Electra::FParamDict& InOptions)
 	#ifdef WINHTTP_FLAG_SECURE_PROTOCOL_TLS1_2
 		SecureProtocols |= WINHTTP_FLAG_SECURE_PROTOCOL_TLS1_2;
 	#endif
+/*
+Removed for the time being since this is not supported on all versions of Windows 10 yet.
+
 	#ifdef WINHTTP_FLAG_SECURE_PROTOCOL_TLS1_3
 		SecureProtocols |= WINHTTP_FLAG_SECURE_PROTOCOL_TLS1_3;
 	#endif
+*/
 	if (SecureProtocols == 0) //-V547
 	{
 		SecureProtocols = WINHTTP_FLAG_SECURE_PROTOCOL_ALL & ~(WINHTTP_FLAG_SECURE_PROTOCOL_SSL2 | WINHTTP_FLAG_SECURE_PROTOCOL_SSL3);
