ElectraPlayer: Added DASH low latency Live streaming

diff --git a/Engine/Plugins/Media/ElectraPlayer/Source/ElectraPlayerRuntime/Private/Runtime/Player/mp4/ManifestMP4.cpp b/Engine/Plugins/Media/ElectraPlayer/Source/ElectraPlayerRuntime/Private/Runtime/Player/mp4/ManifestMP4.cpp
--- a/Engine/Plugins/Media/ElectraPlayer/Source/ElectraPlayerRuntime/Private/Runtime/Player/mp4/ManifestMP4.cpp
+++ b/Engine/Plugins/Media/ElectraPlayer/Source/ElectraPlayerRuntime/Private/Runtime/Player/mp4/ManifestMP4.cpp
@@ -173,12 +173,22 @@ FTimeValue FManifestMP4Internal::GetMinBufferTime() const
 	return FTimeValue().SetFromSeconds(2.0);
 }
 
+TSharedPtrTS<IProducerReferenceTimeInfo> FManifestMP4Internal::GetProducerReferenceTimeInfo(int64 ID) const
+{
+	return nullptr;
+}
+
 
 void FManifestMP4Internal::UpdateDynamicRefetchCounter()
 {
 	// No-op.
 }
 
+void FManifestMP4Internal::TriggerClockSync(IManifest::EClockSyncType InClockSyncType)
+{
+	// No-op.
+}
+
 
 //-----------------------------------------------------------------------------
 /**
@@ -481,6 +491,10 @@ void FManifestMP4Internal::FPlayPeriodMP4::SelectStream(const FString& Adaptatio
 	// .....
 }
 
+void FManifestMP4Internal::FPlayPeriodMP4::TriggerInitSegmentPreload(const TArray<FInitSegmentPreload>& InitSegmentsToPreload)
+{
+	// No-op.
+}
 
 //-----------------------------------------------------------------------------
 /**
@@ -677,6 +691,7 @@ void FManifestMP4Internal::FTimelineAssetMP4::LimitSegmentDownloadSize(TSharedPt
 	// Limit the segment download size.
 	// This helps with downloads that might otherwise take too long or keep the connection open for too long (when downloading a large mp4 from start to finish).
 	const int64 MaxSegmentSize = 4 * 1024 * 1024;
+	const int64 MaxSegmentDurationMSec = 2000;
 	if (InOutSegment.IsValid())
 	{
 		FStreamSegmentRequestMP4* Request = static_cast<FStreamSegmentRequestMP4*>(InOutSegment.Get());
@@ -690,6 +705,7 @@ void FManifestMP4Internal::FTimelineAssetMP4::LimitSegmentDownloadSize(TSharedPt
 		int64 TrackDur = 0;
 		int64 LastTrackOffset = -1;
 		int64 LastSampleSize = 0;
+		int64 TrackDurationLimit = -1;
 		while(AllTrackIterator.IsValid())
 		{
 			const IParserISO14496_12::ITrackIterator* CurrentTrackIt = AllTrackIterator->Current();
@@ -702,13 +718,19 @@ void FManifestMP4Internal::FTimelineAssetMP4::LimitSegmentDownloadSize(TSharedPt
 					bFirst = false;
 					TrackId = CurrentTrackIt->GetTrack()->GetID();
 					TrackTimeScale = CurrentTrackIt->GetTimescale();
+					if (MaxSegmentDurationMSec > 0)
+					{
+						TrackDurationLimit = MaxSegmentDurationMSec * TrackTimeScale / 1000;
+					}
 				}
 				if (TrackId == CurrentTrackIt->GetTrack()->GetID())
 				{
 					TrackDur += CurrentTrackIt->GetDuration();
 				}
 				int64 CurrentTrackOffset = LastTrackOffset;
-				if (CurrentTrackOffset >= EndOffset || CurrentTrackOffset - StartOffset >= MaxSegmentSize)
+				if (CurrentTrackOffset >= EndOffset || 
+					CurrentTrackOffset - StartOffset >= MaxSegmentSize ||
+					(TrackDurationLimit > 0 && TrackDur > TrackDurationLimit))
 				{
 					// Limit reached.
 					Request->FileEndOffset = CurrentTrackOffset - 1;
diff --git a/Engine/Plugins/Media/ElectraPlayer/Source/ElectraPlayerRuntime/Private/Runtime/Player/mp4/ManifestMP4.h b/Engine/Plugins/Media/ElectraPlayer/Source/ElectraPlayerRuntime/Private/Runtime/Player/mp4/ManifestMP4.h
--- a/Engine/Plugins/Media/ElectraPlayer/Source/ElectraPlayerRuntime/Private/Runtime/Player/mp4/ManifestMP4.h
+++ b/Engine/Plugins/Media/ElectraPlayer/Source/ElectraPlayerRuntime/Private/Runtime/Player/mp4/ManifestMP4.h
@@ -30,34 +30,40 @@ public:
 
 	FErrorDetail Build(TSharedPtrTS<IParserISO14496_12> MP4Parser, const FString& URL, const HTTP::FConnectionInfo& InConnectionInfo);
 
-	virtual EType GetPresentationType() const override;
-	virtual FTimeValue GetAnchorTime() const override
+	EType GetPresentationType() const override;
+	TSharedPtrTS<const FLowLatencyDescriptor> GetLowLatencyDescriptor() const override
+	{ return nullptr; }
+	FTimeValue GetAnchorTime() const override
 	{ return FTimeValue::GetZero(); }
-	virtual FTimeRange GetTotalTimeRange() const override
+	FTimeRange GetTotalTimeRange() const override
 	{ return MediaAsset.IsValid() ? MediaAsset->GetTimeRange() : FTimeRange(); }
-	virtual FTimeRange GetSeekableTimeRange() const override
+	FTimeRange GetSeekableTimeRange() const override
 	{
 		// For the time being the seekable range equals the total range.
 		return GetTotalTimeRange();
 	}
-	virtual FTimeRange GetPlaybackRange() const override;
-	virtual void GetSeekablePositions(TArray<FTimespan>& OutPositions) const override
+	FTimeRange GetPlaybackRange() const override;
+	void GetSeekablePositions(TArray<FTimespan>& OutPositions) const override
 	{
 		// For the time being we do not return anything here as that would require to iterate the tracks.
 	}
-	virtual FTimeValue GetDuration() const override
+	FTimeValue GetDuration() const override
 	{ return MediaAsset.IsValid() ? MediaAsset->GetDuration() : FTimeValue(); }
-	virtual FTimeValue GetDefaultStartTime() const override
+	FTimeValue GetDefaultStartTime() const override
 	{ return DefaultStartTime; }
-	virtual void ClearDefaultStartTime() override
+	void ClearDefaultStartTime() override
 	{ DefaultStartTime.SetToInvalid(); }
-	virtual void GetTrackMetadata(TArray<FTrackMetadata>& OutMetadata, EStreamType StreamType) const override;
-	virtual FTimeValue GetMinBufferTime() const override;
-	virtual void UpdateDynamicRefetchCounter() override;
+	void GetTrackMetadata(TArray<FTrackMetadata>& OutMetadata, EStreamType StreamType) const override;
+	FTimeValue GetMinBufferTime() const override;
+	FTimeValue GetDesiredLiveLatency() const override
+	{ return FTimeValue(); }
+	TSharedPtrTS<IProducerReferenceTimeInfo> GetProducerReferenceTimeInfo(int64 ID) const override;
+	void UpdateDynamicRefetchCounter() override;
+	void TriggerClockSync(EClockSyncType InClockSyncType) override;
 
-	virtual IStreamReader* CreateStreamReaderHandler() override;
-	virtual FResult FindPlayPeriod(TSharedPtrTS<IPlayPeriod>& OutPlayPeriod, const FPlayStartPosition& StartPosition, ESearchType SearchType) override;
-	virtual FResult FindNextPlayPeriod(TSharedPtrTS<IPlayPeriod>& OutPlayPeriod, TSharedPtrTS<const IStreamSegment> CurrentSegment) override;
+	IStreamReader* CreateStreamReaderHandler() override;
+	FResult FindPlayPeriod(TSharedPtrTS<IPlayPeriod>& OutPlayPeriod, const FPlayStartPosition& StartPosition, ESearchType SearchType) override;
+	FResult FindNextPlayPeriod(TSharedPtrTS<IPlayPeriod>& OutPlayPeriod, TSharedPtrTS<const IStreamSegment> CurrentSegment) override;
 
 
 	class FRepresentationMP4 : public IPlaybackAssetRepresentation
@@ -68,15 +74,15 @@ public:
 
 		FErrorDetail CreateFrom(const IParserISO14496_12::ITrack* InTrack, const FString& URL);
 
-		virtual FString GetUniqueIdentifier() const override
+		FString GetUniqueIdentifier() const override
 		{ return UniqueIdentifier; }
-		virtual const FStreamCodecInformation& GetCodecInformation() const override
+		const FStreamCodecInformation& GetCodecInformation() const override
 		{ return CodecInformation; }
-		virtual int32 GetBitrate() const override
+		int32 GetBitrate() const override
 		{ return Bitrate; }
-		virtual int32 GetQualityIndex() const override
+		int32 GetQualityIndex() const override
 		{ return 0; }
-		virtual bool CanBePlayed() const override
+		bool CanBePlayed() const override
 		{ return true; }
 
 		const FString& GetName() const
@@ -99,17 +105,19 @@ public:
 
 		FErrorDetail CreateFrom(const IParserISO14496_12::ITrack* InTrack, const FString& URL);
 
-		virtual FString GetUniqueIdentifier() const override
+		FString GetUniqueIdentifier() const override
 		{ return UniqueIdentifier; }
-		virtual FString GetListOfCodecs() const override
+		FString GetListOfCodecs() const override
 		{ return CodecRFC6381; }
-		virtual FString GetLanguage() const override
+		FString GetLanguage() const override
 		{ return Language; }
-		virtual int32 GetNumberOfRepresentations() const override
+		int32 GetNumberOfRepresentations() const override
 		{ return Representation.IsValid() ? 1 : 0; }
-		virtual TSharedPtrTS<IPlaybackAssetRepresentation> GetRepresentationByIndex(int32 RepresentationIndex) const override
+		bool IsLowLatencyEnabled() const override
+		{ return false; }
+		TSharedPtrTS<IPlaybackAssetRepresentation> GetRepresentationByIndex(int32 RepresentationIndex) const override
 		{ return RepresentationIndex == 0 ? Representation : TSharedPtrTS<IPlaybackAssetRepresentation>(); }
-		virtual TSharedPtrTS<IPlaybackAssetRepresentation> GetRepresentationByUniqueIdentifier(const FString& InUniqueIdentifier) const override
+		TSharedPtrTS<IPlaybackAssetRepresentation> GetRepresentationByUniqueIdentifier(const FString& InUniqueIdentifier) const override
 		{ return Representation.IsValid() && Representation->GetUniqueIdentifier() == InUniqueIdentifier ? Representation : TSharedPtrTS<IPlaybackAssetRepresentation>(); }
 	private:
 		TSharedPtrTS<FRepresentationMP4>		Representation;
@@ -267,23 +275,24 @@ public:
 	public:
 		FPlayPeriodMP4(TSharedPtrTS<FTimelineAssetMP4> InMediaAsset);
 		virtual ~FPlayPeriodMP4();
-		virtual void SetStreamPreferences(EStreamType ForStreamType, const FStreamSelectionAttributes& StreamAttributes) override;
-		virtual EReadyState GetReadyState() override;
-		virtual void Load() override;
-		virtual void PrepareForPlay() override;
-		virtual int64 GetDefaultStartingBitrate() const override;
-		virtual TSharedPtrTS<FBufferSourceInfo> GetSelectedStreamBufferSourceInfo(EStreamType StreamType) override;
-		virtual FString GetSelectedAdaptationSetID(EStreamType StreamType) override;
-		virtual ETrackChangeResult ChangeTrackStreamPreference(EStreamType ForStreamType, const FStreamSelectionAttributes& StreamAttributes) override;
-		virtual TSharedPtrTS<ITimelineMediaAsset> GetMediaAsset() const override;
-		virtual void SelectStream(const FString& AdaptationSetID, const FString& RepresentationID) override;
-		virtual FResult GetStartingSegment(TSharedPtrTS<IStreamSegment>& OutSegment, const FPlayerSequenceState& InSequenceState, const FPlayStartPosition& StartPosition, ESearchType SearchType) override;
-		virtual FResult GetContinuationSegment(TSharedPtrTS<IStreamSegment>& OutSegment, EStreamType StreamType, const FPlayerSequenceState& LoopState, const FPlayStartPosition& StartPosition, ESearchType SearchType) override;
-		virtual FResult GetNextSegment(TSharedPtrTS<IStreamSegment>& OutSegment, TSharedPtrTS<const IStreamSegment> CurrentSegment, const FPlayStartOptions& Options) override;
-		virtual FResult GetRetrySegment(TSharedPtrTS<IStreamSegment>& OutSegment, TSharedPtrTS<const IStreamSegment> CurrentSegment, const FPlayStartOptions& Options, bool bReplaceWithFillerData) override;
-		virtual FResult GetLoopingSegment(TSharedPtrTS<IStreamSegment>& OutSegment, const FPlayerSequenceState& InSequenceState, const TMultiMap<EStreamType, TSharedPtrTS<IStreamSegment>>& InFinishedSegments, const FPlayStartPosition& StartPosition, ESearchType SearchType) override;
-		virtual void IncreaseSegmentFetchDelay(const FTimeValue& IncreaseAmount) override;
-		virtual void GetSegmentInformation(TArray<FSegmentInformation>& OutSegmentInformation, FTimeValue& OutAverageSegmentDuration, TSharedPtrTS<const IStreamSegment> CurrentSegment, const FTimeValue& LookAheadTime, const FString& AdaptationSetID, const FString& RepresentationID) override;
+		void SetStreamPreferences(EStreamType ForStreamType, const FStreamSelectionAttributes& StreamAttributes) override;
+		EReadyState GetReadyState() override;
+		void Load() override;
+		void PrepareForPlay() override;
+		int64 GetDefaultStartingBitrate() const override;
+		TSharedPtrTS<FBufferSourceInfo> GetSelectedStreamBufferSourceInfo(EStreamType StreamType) override;
+		FString GetSelectedAdaptationSetID(EStreamType StreamType) override;
+		ETrackChangeResult ChangeTrackStreamPreference(EStreamType ForStreamType, const FStreamSelectionAttributes& StreamAttributes) override;
+		TSharedPtrTS<ITimelineMediaAsset> GetMediaAsset() const override;
+		void SelectStream(const FString& AdaptationSetID, const FString& RepresentationID) override;
+		void TriggerInitSegmentPreload(const TArray<FInitSegmentPreload>& InitSegmentsToPreload) override;
+		FResult GetStartingSegment(TSharedPtrTS<IStreamSegment>& OutSegment, const FPlayerSequenceState& InSequenceState, const FPlayStartPosition& StartPosition, ESearchType SearchType) override;
+		FResult GetContinuationSegment(TSharedPtrTS<IStreamSegment>& OutSegment, EStreamType StreamType, const FPlayerSequenceState& LoopState, const FPlayStartPosition& StartPosition, ESearchType SearchType) override;
+		FResult GetNextSegment(TSharedPtrTS<IStreamSegment>& OutSegment, TSharedPtrTS<const IStreamSegment> CurrentSegment, const FPlayStartOptions& Options) override;
+		FResult GetRetrySegment(TSharedPtrTS<IStreamSegment>& OutSegment, TSharedPtrTS<const IStreamSegment> CurrentSegment, const FPlayStartOptions& Options, bool bReplaceWithFillerData) override;
+		FResult GetLoopingSegment(TSharedPtrTS<IStreamSegment>& OutSegment, const FPlayerSequenceState& InSequenceState, const TMultiMap<EStreamType, TSharedPtrTS<IStreamSegment>>& InFinishedSegments, const FPlayStartPosition& StartPosition, ESearchType SearchType) override;
+		void IncreaseSegmentFetchDelay(const FTimeValue& IncreaseAmount) override;
+		void GetSegmentInformation(TArray<FSegmentInformation>& OutSegmentInformation, FTimeValue& OutAverageSegmentDuration, TSharedPtrTS<const IStreamSegment> CurrentSegment, const FTimeValue& LookAheadTime, const FString& AdaptationSetID, const FString& RepresentationID) override;
 
 	private:
 		void SelectInitialStream(EStreamType StreamType);
diff --git a/Engine/Plugins/Media/ElectraPlayer/Source/ElectraPlayerRuntime/Private/Runtime/Player/mp4/StreamReaderMP4.cpp b/Engine/Plugins/Media/ElectraPlayer/Source/ElectraPlayerRuntime/Private/Runtime/Player/mp4/StreamReaderMP4.cpp
--- a/Engine/Plugins/Media/ElectraPlayer/Source/ElectraPlayerRuntime/Private/Runtime/Player/mp4/StreamReaderMP4.cpp
+++ b/Engine/Plugins/Media/ElectraPlayer/Source/ElectraPlayerRuntime/Private/Runtime/Player/mp4/StreamReaderMP4.cpp
@@ -61,7 +61,7 @@ EStreamType FStreamSegmentRequestMP4::GetType() const
 	return  PrimaryStreamType;
 }
 
-void FStreamSegmentRequestMP4::SetExecutionDelay(const FTimeValue& ExecutionDelay)
+void FStreamSegmentRequestMP4::SetExecutionDelay(const FTimeValue& UTCNow, const FTimeValue& ExecutionDelay)
 {
 	// No op for mp4.
 }
@@ -379,13 +379,11 @@ void FStreamReaderMP4::HandleRequest()
 	ds.TimeToDownload      = 0.0;
 	ds.ByteSize 		   = -1;
 	ds.NumBytesDownloaded  = 0;
-	ds.ThroughputBps	   = 0;
 	ds.bInsertedFillerData = false;
 	ds.URL  			   = TimelineAsset->GetMediaURL();
 	ds.bIsMissingSegment   = false;
 	ds.bParseFailure	   = false;
 	ds.RetryNumber  	   = Request->NumOverallRetries;
-	ds.ABRState.Reset();
 
 	Parameters.EventListener->OnFragmentOpen(Request);
 
@@ -714,11 +712,6 @@ void FStreamReaderMP4::HandleRequest()
 	ds.TimeToDownload     = (Request->ConnectionInfo.RequestEndTime - Request->ConnectionInfo.RequestStartTime).GetAsSeconds();
 	ds.ByteSize 		  = Request->ConnectionInfo.ContentLength;
 	ds.NumBytesDownloaded = Request->ConnectionInfo.BytesReadSoFar;
-	ds.ThroughputBps	  = Request->ConnectionInfo.Throughput.GetThroughput();
-	if (ds.ThroughputBps == 0)
-	{
-		ds.ThroughputBps = ds.TimeToDownload > 0.0 ? 8 * ds.NumBytesDownloaded / ds.TimeToDownload : 0;
-	}
 	ds.bIsCachedResponse = Request->ConnectionInfo.bIsCachedResponse;
 
 	ActiveTrackMap.Reset();
diff --git a/Engine/Plugins/Media/ElectraPlayer/Source/ElectraPlayerRuntime/Private/Runtime/Player/mp4/StreamReaderMP4.h b/Engine/Plugins/Media/ElectraPlayer/Source/ElectraPlayerRuntime/Private/Runtime/Player/mp4/StreamReaderMP4.h
--- a/Engine/Plugins/Media/ElectraPlayer/Source/ElectraPlayerRuntime/Private/Runtime/Player/mp4/StreamReaderMP4.h
+++ b/Engine/Plugins/Media/ElectraPlayer/Source/ElectraPlayerRuntime/Private/Runtime/Player/mp4/StreamReaderMP4.h
@@ -28,25 +28,27 @@ public:
 	FStreamSegmentRequestMP4();
 	virtual ~FStreamSegmentRequestMP4();
 
-	virtual void SetPlaybackSequenceID(uint32 InPlaybackSequenceID) override;
-	virtual uint32 GetPlaybackSequenceID() const override;
+	void SetPlaybackSequenceID(uint32 InPlaybackSequenceID) override;
+	uint32 GetPlaybackSequenceID() const override;
 
-	virtual void SetExecutionDelay(const FTimeValue& ExecutionDelay) override;
-	virtual FTimeValue GetExecuteAtUTCTime() const override;
+	void SetExecutionDelay(const FTimeValue& UTCNow, const FTimeValue& ExecutionDelay) override;
+	FTimeValue GetExecuteAtUTCTime() const override;
 
-	virtual EStreamType GetType() const override;
+	EStreamType GetType() const override;
 
-	virtual void GetDependentStreams(TArray<TSharedPtrTS<IStreamSegment>>& OutDependentStreams) const override;
-	virtual void GetRequestedStreams(TArray<TSharedPtrTS<IStreamSegment>>& OutRequestedStreams) override;
-	virtual void GetEndedStreams(TArray<TSharedPtrTS<IStreamSegment>>& OutAlreadyEndedStreams) override;
+	void GetDependentStreams(TArray<TSharedPtrTS<IStreamSegment>>& OutDependentStreams) const override;
+	void GetRequestedStreams(TArray<TSharedPtrTS<IStreamSegment>>& OutRequestedStreams) override;
+	void GetEndedStreams(TArray<TSharedPtrTS<IStreamSegment>>& OutAlreadyEndedStreams) override;
 
 	//! Returns the first PTS value as indicated by the media timeline. This should correspond to the actual absolute PTS of the sample.
-	virtual FTimeValue GetFirstPTS() const override;
+	FTimeValue GetFirstPTS() const override;
 
-	virtual int32 GetQualityIndex() const override;
-	virtual int32 GetBitrate() const override;
+	int32 GetQualityIndex() const override;
+	int32 GetBitrate() const override;
 
-	virtual void GetDownloadStats(Metrics::FSegmentDownloadStats& OutStats) const override;
+	void GetDownloadStats(Metrics::FSegmentDownloadStats& OutStats) const override;
+	bool GetStartupDelay(FTimeValue& OutStartTime, FTimeValue& OutTimeIntoSegment, FTimeValue& OutSegmentDuration) const override
+	{ return false; }
 
 
 
diff --git a/Engine/Plugins/Media/ElectraPlayer/Source/ElectraPlayerRuntime/Private/Runtime/Renderer/RendererAudio.cpp b/Engine/Plugins/Media/ElectraPlayer/Source/ElectraPlayerRuntime/Private/Runtime/Renderer/RendererAudio.cpp
--- a/Engine/Plugins/Media/ElectraPlayer/Source/ElectraPlayerRuntime/Private/Runtime/Renderer/RendererAudio.cpp
+++ b/Engine/Plugins/Media/ElectraPlayer/Source/ElectraPlayerRuntime/Private/Runtime/Renderer/RendererAudio.cpp
@@ -9,14 +9,16 @@
 #include "ElectraPlayer.h"
 #include "ElectraPlayerPrivate.h"
 
-void FElectraRendererAudio::SampleReleasedToPool(FTimespan DurationToRelease)
+void FElectraRendererAudio::SampleReleasedToPool(IDecoderOutput* InDecoderOutput)
 {
 	FPlatformAtomics::InterlockedDecrement(&NumOutputAudioBuffersInUse);
 	check(NumOutputAudioBuffersInUse >= 0 && NumOutputAudioBuffersInUse <= NumBuffers)
-}
 
-FElectraRendererAudio::SystemConfiguration::SystemConfiguration()
-{
+	TSharedPtr<IMediaRenderer, ESPMode::ThreadSafe> Parent = ParentRenderer.Pin();
+	if (Parent.IsValid())
+	{
+		Parent->SampleReleasedToPool(InDecoderOutput);
+	}
 }
 
 bool FElectraRendererAudio::Startup(const FElectraRendererAudio::SystemConfiguration& Configuration)
@@ -174,12 +176,13 @@ UEMediaError FElectraRendererAudio::ReturnBuffer(IBuffer* Buffer, bool bRender,
 		uint32 InSampleRate = (uint32)variantSampleRate.GetInt64();
 		uint32 InUsedBufferBytes = (uint32)variantBufferUsedBytes.GetInt64();
 
-		FTimespan InDuration = FTimespan(variantDuration.GetTimeValue().GetAsHNS());
+		FTimespan InDuration = variantDuration.GetTimeValue().GetAsTimespan();
 		FTimespan InPts = variantPts.GetTimeValue().GetAsTimespan();
 		int64 InSequenceIndex = variantPts.GetTimeValue().GetSequenceIndex();
 
 		//UE_LOG(LogElectraPlayer, VeryVerbose, TEXT("-- FElectraRendererAudio::ReturnBuffer: Audio sample for time %s"), *InPts.ToString(TEXT("%h:%m:%s.%f")));
 
+		DecoderOutput->GetMutablePropertyDictionary() = InSampleProperties;
 		DecoderOutput->Initialize(IAudioDecoderOutput::ESampleFormat::Int16, InNumChannels, InSampleRate, InDuration, FDecoderTimeStamp(InPts, InSequenceIndex), InUsedBufferBytes);
 
 		// Push buffer to output queue...
@@ -198,15 +201,6 @@ UEMediaError FElectraRendererAudio::ReturnBuffer(IBuffer* Buffer, bool bRender,
 	return UEMEDIA_ERROR_OK;
 }
 
-void FElectraRendererAudio::SetRenderClock(const FTimespan& RenderTime)
-{
-	// Latest audio as time reference
-	Electra::FTimeValue AudioRenderTime;
-	AudioRenderTime.SetFromHNS(RenderTime.GetTicks());
-	RenderClock->SetCurrentTime(Electra::IMediaRenderClock::ERendererType::Audio, AudioRenderTime);
-	RenderClock->SetCurrentTime(Electra::IMediaRenderClock::ERendererType::Video, AudioRenderTime);
-}
-
 UEMediaError FElectraRendererAudio::ReleaseBufferPool()
 {
 	DecoderOutputPool.Reset();
@@ -231,6 +225,12 @@ void FElectraRendererAudio::SetRenderClock(TSharedPtr<Electra::IMediaRenderClock
 	RenderClock = InRenderClock;
 }
 
+void FElectraRendererAudio::SetParentRenderer(TWeakPtr<IMediaRenderer, ESPMode::ThreadSafe> InParentRenderer)
+{
+	ParentRenderer = MoveTemp(InParentRenderer);
+}
+
+
 void FElectraRendererAudio::SetNextApproximatePresentationTime(const Electra::FTimeValue& NextApproxPTS)
 {
 }
@@ -272,5 +272,3 @@ void FElectraRendererAudio::StopRendering(const Electra::FParamDict& InOptions)
 		PinnedPlayer->OnAudioRenderingStopped();
 	}
 }
-
-
diff --git a/Engine/Plugins/Media/ElectraPlayer/Source/ElectraPlayerRuntime/Private/Runtime/Renderer/RendererAudio.h b/Engine/Plugins/Media/ElectraPlayer/Source/ElectraPlayerRuntime/Private/Runtime/Renderer/RendererAudio.h
--- a/Engine/Plugins/Media/ElectraPlayer/Source/ElectraPlayerRuntime/Private/Runtime/Renderer/RendererAudio.h
+++ b/Engine/Plugins/Media/ElectraPlayer/Source/ElectraPlayerRuntime/Private/Runtime/Renderer/RendererAudio.h
@@ -26,7 +26,6 @@ class FElectraRendererAudio : public Electra::IMediaRenderer, public TSharedFrom
 public:
 	struct SystemConfiguration
 	{
-		SystemConfiguration();
 	};
 	static bool Startup(const SystemConfiguration& Configuration);
 	static void Shutdown(void);
@@ -36,10 +35,6 @@ public:
 
 	struct InstanceConfiguration
 	{
-		InstanceConfiguration()
-		{
-			// todo
-		}
 	};
 
 	enum OpenError
@@ -62,47 +57,48 @@ public:
 	//
 
 	// Returns the properties of the buffer pool. Those properties should not change
-	virtual const Electra::FParamDict& GetBufferPoolProperties() const override;
+	const Electra::FParamDict& GetBufferPoolProperties() const override;
 
 	// Create a buffer pool from where a decoder can get the block of memory to decode into.
-	virtual UEMediaError CreateBufferPool(const Electra::FParamDict& Parameters) override;
+	UEMediaError CreateBufferPool(const Electra::FParamDict& Parameters) override;
 
 	// Asks for a sample buffer from the buffer pool created previously through CreateBufferPool()
-	virtual UEMediaError AcquireBuffer(IBuffer*& OutBuffer, int32 TimeoutInMicroseconds, const Electra::FParamDict& InParameters) override;
+	UEMediaError AcquireBuffer(IBuffer*& OutBuffer, int32 TimeoutInMicroseconds, const Electra::FParamDict& InParameters) override;
 
 	// Releases the buffer for rendering and subsequent return to the buffer pool
-	virtual UEMediaError ReturnBuffer(IBuffer* Buffer, bool bRender, const Electra::FParamDict& InSampleProperties) override;
+	UEMediaError ReturnBuffer(IBuffer* Buffer, bool bRender, const Electra::FParamDict& InSampleProperties) override;
 
 	// Informs that the decoder is done with this pool. NO FREE!!!
-	virtual UEMediaError ReleaseBufferPool() override;
+	UEMediaError ReleaseBufferPool() override;
 
-	virtual bool CanReceiveOutputFrames(uint64 NumFrames) const override;
+	bool CanReceiveOutputFrames(uint64 NumFrames) const override;
 
 	// Receives the render clock we need to update with the most recently rendered sample's timestamp.
-	virtual void SetRenderClock(TSharedPtr<Electra::IMediaRenderClock, ESPMode::ThreadSafe> RenderClock) override;
+	void SetRenderClock(TSharedPtr<Electra::IMediaRenderClock, ESPMode::ThreadSafe> RenderClock) override;
+
+	// Called if this renderer is being wrapped by another renderer.
+	void SetParentRenderer(TWeakPtr<IMediaRenderer, ESPMode::ThreadSafe> ParentRenderer) override;
 
 	// Sets the next expected sample's approximate presentation time stamp
-	virtual void SetNextApproximatePresentationTime(const Electra::FTimeValue& NextApproxPTS) override;
+	void SetNextApproximatePresentationTime(const Electra::FTimeValue& NextApproxPTS) override;
 
 	// Flushes all pending buffers not yet rendered
-	virtual UEMediaError Flush(const Electra::FParamDict& InOptions) override;
+	UEMediaError Flush(const Electra::FParamDict& InOptions) override;
 
 	// Begins rendering of the first sample buffer
-	virtual void StartRendering(const Electra::FParamDict& InOptions) override;
+	void StartRendering(const Electra::FParamDict& InOptions) override;
 
 	// Stops rendering of sample buffers
-	virtual void StopRendering(const Electra::FParamDict& InOptions) override;
-
+	void StopRendering(const Electra::FParamDict& InOptions) override;
 
-	//! Called when audio sample is returned to pool
-	virtual void SampleReleasedToPool(FTimespan DurationToRelease) override;
-
-	//! Set used clock to some external value
-	void SetRenderClock(const FTimespan& RenderTime);
+	// Called when audio sample is returned to pool
+	void SampleReleasedToPool(IDecoderOutput* InDecoderOutput) override;
 
 private:
 
 	TWeakPtr<FElectraPlayer, ESPMode::ThreadSafe> Player;
+	TWeakPtr<IMediaRenderer, ESPMode::ThreadSafe> ParentRenderer;
+
 	TDecoderOutputObjectPool<IAudioDecoderOutput, FElectraPlayerPlatformAudioDecoderOutputFactory> DecoderOutputPool;
 
 	uint32					 MaxBufferSize;
@@ -113,27 +109,28 @@ private:
 	public:
 		explicit FMediaBufferSharedPtrWrapper(IAudioDecoderOutputPtr InDecoderOutput)
 			: DecoderOutput(InDecoderOutput)
-		{
-		}
+		{ }
+
+		virtual ~FMediaBufferSharedPtrWrapper() = default;
 
-		~FMediaBufferSharedPtrWrapper()
+		const Electra::FParamDict& GetBufferProperties() const override
 		{
+			return BufferProperties;
 		}
-
-		virtual const Electra::FParamDict& GetBufferProperties() const override
+		Electra::FParamDict& GetMutableBufferProperties() override
 		{
 			return BufferProperties;
 		}
+
 		Electra::FParamDict BufferProperties;
 		IAudioDecoderOutputPtr DecoderOutput;
 	};
 
 	// ------------------------------------------------------------------------------------------------------
 
-	Electra::FParamDict							BufferPoolProperties;
-	int32										NumBuffers;
-	int32										NumBuffersAcquiredForDecoder;
-
+	Electra::FParamDict BufferPoolProperties;
+	int32 NumBuffers;
+	int32 NumBuffersAcquiredForDecoder;
 	TSharedPtr<Electra::IMediaRenderClock, ESPMode::ThreadSafe>	RenderClock;
 };
 
diff --git a/Engine/Plugins/Media/ElectraPlayer/Source/ElectraPlayerRuntime/Private/Runtime/Renderer/RendererBase.h b/Engine/Plugins/Media/ElectraPlayer/Source/ElectraPlayerRuntime/Private/Runtime/Renderer/RendererBase.h
--- a/Engine/Plugins/Media/ElectraPlayer/Source/ElectraPlayerRuntime/Private/Runtime/Renderer/RendererBase.h
+++ b/Engine/Plugins/Media/ElectraPlayer/Source/ElectraPlayerRuntime/Private/Runtime/Renderer/RendererBase.h
@@ -68,6 +68,9 @@ namespace Electra
 			 * @return Const reference to the dictionary containing buffer properties.
 			 */
 			virtual const FParamDict& GetBufferProperties() const = 0;
+			
+			// Same, but returns the writable dictionary. Use with caution.
+			virtual FParamDict& GetMutableBufferProperties() = 0;
 		};
 
 
@@ -161,6 +164,12 @@ namespace Electra
 		 */
 		virtual void SetRenderClock(TSharedPtr<IMediaRenderClock, ESPMode::ThreadSafe> RenderClock) = 0;
 
+
+		/**
+		 * Called if this renderer is being wrapped by another renderer.
+		 */
+		virtual void SetParentRenderer(TWeakPtr<IMediaRenderer, ESPMode::ThreadSafe> ParentRenderer) = 0;
+
 		/**
 		 * Sets the next expected sample's approximate presentation time stamp.
 		 * This is called as data buffering begins at the indicated play start
diff --git a/Engine/Plugins/Media/ElectraPlayer/Source/ElectraPlayerRuntime/Private/Runtime/Renderer/RendererVideo.cpp b/Engine/Plugins/Media/ElectraPlayer/Source/ElectraPlayerRuntime/Private/Runtime/Renderer/RendererVideo.cpp
--- a/Engine/Plugins/Media/ElectraPlayer/Source/ElectraPlayerRuntime/Private/Runtime/Renderer/RendererVideo.cpp
+++ b/Engine/Plugins/Media/ElectraPlayer/Source/ElectraPlayerRuntime/Private/Runtime/Renderer/RendererVideo.cpp
@@ -16,9 +16,15 @@ DECLARE_DWORD_COUNTER_STAT(TEXT("Electra::MEDIArendererVideoUE_NumUsedOutputVide
 
 // -----------------------------------------------------------------------------------------------------------------------
 
-void FElectraRendererVideo::SampleReleasedToPool(FTimespan durationToRelease)
+void FElectraRendererVideo::SampleReleasedToPool(IDecoderOutput* InDecoderOutput)
 {
 	FPlatformAtomics::InterlockedDecrement(&NumOutputTexturesInUse);
+
+	TSharedPtr<IMediaRenderer, ESPMode::ThreadSafe> Parent = ParentRenderer.Pin();
+	if (Parent.IsValid())
+	{
+		Parent->SampleReleasedToPool(InDecoderOutput);
+	}
 }
 
 
@@ -188,6 +194,7 @@ UEMediaError FElectraRendererVideo::ReturnBuffer(IBuffer* Buffer, bool bRender,
 	}
 
 	FMediaBufferSharedPtrWrapper* MediaBufferSharedPtrWrapper = static_cast<FMediaBufferSharedPtrWrapper*>(Buffer);
+	MediaBufferSharedPtrWrapper->DecoderOutput->GetMutablePropertyDictionary() = InSampleProperties;
 
 	if (bRender)
 	{
@@ -239,6 +246,11 @@ void FElectraRendererVideo::SetRenderClock(TSharedPtr<Electra::IMediaRenderClock
 	RenderClock = InRenderClock;
 }
 
+void FElectraRendererVideo::SetParentRenderer(TWeakPtr<IMediaRenderer, ESPMode::ThreadSafe> InParentRenderer)
+{
+	ParentRenderer = MoveTemp(InParentRenderer);
+}
+
 /**
  *
  */
diff --git a/Engine/Plugins/Media/ElectraPlayer/Source/ElectraPlayerRuntime/Private/Runtime/Renderer/RendererVideo.h b/Engine/Plugins/Media/ElectraPlayer/Source/ElectraPlayerRuntime/Private/Runtime/Renderer/RendererVideo.h
--- a/Engine/Plugins/Media/ElectraPlayer/Source/ElectraPlayerRuntime/Private/Runtime/Renderer/RendererVideo.h
+++ b/Engine/Plugins/Media/ElectraPlayer/Source/ElectraPlayerRuntime/Private/Runtime/Renderer/RendererVideo.h
@@ -46,7 +46,7 @@ public:
 		eOpen_Error,
 	};
 
-	virtual OpenError Open(const InstanceConfiguration& Config);
+	OpenError Open(const InstanceConfiguration& Config);
 	void Close(void);
 
 	void PrepareForDecoderShutdown();
@@ -63,47 +63,51 @@ public:
 	//
 
 	// Returns the properties of the buffer pool. Those properties should not change
-	virtual const Electra::FParamDict& GetBufferPoolProperties() const override;
+	const Electra::FParamDict& GetBufferPoolProperties() const override;
 
 	// Create a buffer pool from where a decoder can get the block of memory to decode into.
-	virtual UEMediaError CreateBufferPool(const Electra::FParamDict& Parameters) override;
+	UEMediaError CreateBufferPool(const Electra::FParamDict& Parameters) override;
 
 	// Asks for a sample buffer from the buffer pool created previously through CreateBufferPool()
-	virtual UEMediaError AcquireBuffer(IBuffer*& OutBuffer, int32 TimeoutInMicroseconds, const Electra::FParamDict& InParameters) override;
+	UEMediaError AcquireBuffer(IBuffer*& OutBuffer, int32 TimeoutInMicroseconds, const Electra::FParamDict& InParameters) override;
 
 	// Releases the buffer for rendering and subsequent return to the buffer pool
-	virtual UEMediaError ReturnBuffer(IBuffer* Buffer, bool bRender, const Electra::FParamDict& InSampleProperties) override;
+	UEMediaError ReturnBuffer(IBuffer* Buffer, bool bRender, const Electra::FParamDict& InSampleProperties) override;
 
 	// Informs that the decoder is done with this pool. NO FREE!!!
-	virtual UEMediaError ReleaseBufferPool() override;
+	UEMediaError ReleaseBufferPool() override;
 
-	virtual bool CanReceiveOutputFrames(uint64 NumFrames) const override;
+	bool CanReceiveOutputFrames(uint64 NumFrames) const override;
 
 	// Receives the render clock we need to update with the most recently rendered sample's timestamp.
-	virtual void SetRenderClock(TSharedPtr<Electra::IMediaRenderClock, ESPMode::ThreadSafe> RenderClock) override;
+	void SetRenderClock(TSharedPtr<Electra::IMediaRenderClock, ESPMode::ThreadSafe> RenderClock) override;
+
+	// Called if this renderer is being wrapped by another renderer.
+	void SetParentRenderer(TWeakPtr<IMediaRenderer, ESPMode::ThreadSafe> ParentRenderer) override;
 
 	// Sets the next expected sample's approximate presentation time stamp
-	virtual void SetNextApproximatePresentationTime(const Electra::FTimeValue& NextApproxPTS) override;
+	void SetNextApproximatePresentationTime(const Electra::FTimeValue& NextApproxPTS) override;
 
 	// Flushes all pending buffers not yet rendered
-	virtual UEMediaError Flush(const Electra::FParamDict& InOptions) override;
+	UEMediaError Flush(const Electra::FParamDict& InOptions) override;
 
 	// Begins rendering of the first sample buffer
-	virtual void StartRendering(const Electra::FParamDict& InOptions) override;
+	void StartRendering(const Electra::FParamDict& InOptions) override;
 
 	// Stops rendering of sample buffers
-	virtual void StopRendering(const Electra::FParamDict& InOptions) override;
+	void StopRendering(const Electra::FParamDict& InOptions) override;
 
 	//! Tick any output buffer logic that might need regular updates
-	virtual void TickOutputBufferPool() override;
+	void TickOutputBufferPool() override;
 
-	//! Called from FElectraPlayerVideoDecoderOutputXYZ when texture is returned to pool
-	virtual void SampleReleasedToPool(FTimespan DurationToRelease) override;
+	//! Called from FElectraPlayerVideoDecoderOutput[Platform] when texture is returned to pool
+	void SampleReleasedToPool(IDecoderOutput* InDecoderOutput) override;
 
 private:
 
 
 	TWeakPtr<FElectraPlayer, ESPMode::ThreadSafe> Player;
+	TWeakPtr<IMediaRenderer, ESPMode::ThreadSafe> ParentRenderer;
 
 	TDecoderOutputObjectPool<FVideoDecoderOutput, FElectraPlayerPlatformVideoDecoderOutputFactory> DecoderOutputPool;
 	int32 NumOutputTexturesInUse;
@@ -114,17 +118,19 @@ private:
 	public:
 		explicit FMediaBufferSharedPtrWrapper(const FVideoDecoderOutputPtr& InDecoderOutput)
 			: DecoderOutput(InDecoderOutput)
-		{
-		}
+		{ }
 
-		~FMediaBufferSharedPtrWrapper()
+		~FMediaBufferSharedPtrWrapper() = default;
+
+		const Electra::FParamDict& GetBufferProperties() const override
 		{
+			return BufferProperties;
 		}
-
-		virtual const Electra::FParamDict& GetBufferProperties() const override
+		Electra::FParamDict& GetMutableBufferProperties() override
 		{
 			return BufferProperties;
 		}
+
 		Electra::FParamDict		BufferProperties;
 		FVideoDecoderOutputPtr	DecoderOutput;
 	};
@@ -136,18 +142,11 @@ private:
 
 
 	// These queues are all single-producer, single-consumer which is the default mode...
+	TQueue<FVideoDecoderOutputPtr> QueueTickedAndWaitingForDecoder;
 
-	TQueue<FVideoDecoderOutputPtr>					QueueTickedAndWaitingForDecoder;
-
-	// ------------------------------------------------------------------------------------------------------
-
-	Electra::FParamDict								BufferPoolProperties;
-
-	int32											NumBuffers;
-	int32											NumBuffersAcquiredForDecoder;
-
-	FTimespan										DurationOfPendingRenderBuffers;
-
+	Electra::FParamDict BufferPoolProperties;
+	int32 NumBuffers;
+	int32 NumBuffersAcquiredForDecoder;
 	TSharedPtr<Electra::IMediaRenderClock, ESPMode::ThreadSafe>	RenderClock;
 };
 
diff --git a/Engine/Plugins/Media/ElectraPlayer/Source/ElectraPlayerRuntime/Private/Runtime/StreamAccessUnitBuffer.h b/Engine/Plugins/Media/ElectraPlayer/Source/ElectraPlayerRuntime/Private/Runtime/StreamAccessUnitBuffer.h
--- a/Engine/Plugins/Media/ElectraPlayer/Source/ElectraPlayerRuntime/Private/Runtime/StreamAccessUnitBuffer.h
+++ b/Engine/Plugins/Media/ElectraPlayer/Source/ElectraPlayerRuntime/Private/Runtime/StreamAccessUnitBuffer.h
@@ -71,6 +71,7 @@ namespace Electra
 		FTimeValue					EarliestPTS;					//!< Earliest PTS at which to present samples. If this is larger than PTS the sample is not to be presented.
 		FTimeValue					LatestPTS;						//!< Latest PTS at which to present samples. If this is less than PTS the sample is not to be presented.
 		FTimeValue					OffsetFromSegmentStart;			//!< If set, the difference between the first segment AU's PTS and the expected time according to the playlist.
+		FTimeValue					ProducerReferenceTime;			//!< If set, the wallclock time of the producer when this AU was encoded or captured
 		int64						SequenceIndex;
 		uint32						AUSize;							//!< Size of this access unit
 		void*						AUData;							//!< Access unit data
@@ -350,7 +351,7 @@ namespace Electra
 					else
 					{
 						FrontDTS.SetToInvalid();
-						for(uint32 i = 1; i < AccessUnits.Num(); ++i)
+						for(int32 i=1; i<AccessUnits.Num(); ++i)
 						{
 							if (AccessUnits[i]->DropState == FAccessUnit::EDropState::None)
 							{
diff --git a/Engine/Plugins/Media/ElectraUtil/Source/ElectraBase/Private/TimeUtilities.cpp b/Engine/Plugins/Media/ElectraUtil/Source/ElectraBase/Private/TimeUtilities.cpp
--- a/Engine/Plugins/Media/ElectraUtil/Source/ElectraBase/Private/TimeUtilities.cpp
+++ b/Engine/Plugins/Media/ElectraUtil/Source/ElectraBase/Private/TimeUtilities.cpp
@@ -585,6 +585,24 @@ namespace Electra
 	}
 
 
+	namespace RFC5905
+	{
+		bool ParseNTPTime(FTimeValue& OutTimeValue, uint64 NtpTimestampFormat)
+		{
+			if (NtpTimestampFormat)
+			{
+				// Parse an NTP time such that it based on the Unix Epoch.
+				uint64 Seconds = (NtpTimestampFormat >> 32) - 2208988800UL;			// 70 years, with 17 leap year days as seconds (70*365 + 17) * 86400
+				uint64 Nanos = ((NtpTimestampFormat & 0xffffffff) * 1000000) >> 32;
+				int64 HNS = Seconds * 10000000 + Nanos * 10;
+				OutTimeValue.SetFromHNS(HNS);
+				return true;
+			}
+			return false;
+		}
+	}
+
+
 } // namespace Electra
 
 
diff --git a/Engine/Plugins/Media/ElectraUtil/Source/ElectraBase/Private/URLParser.cpp b/Engine/Plugins/Media/ElectraUtil/Source/ElectraBase/Private/URLParser.cpp
--- a/Engine/Plugins/Media/ElectraUtil/Source/ElectraBase/Private/URLParser.cpp
+++ b/Engine/Plugins/Media/ElectraUtil/Source/ElectraBase/Private/URLParser.cpp
@@ -336,6 +336,13 @@ namespace Electra
 		GetPathComponents(OutPathComponents, GetPath());
 	}
 
+	FString FURL_RFC3986::GetLastPathComponent() const
+	{
+		TArray<FString> Components;
+		GetPathComponents(Components, GetPath());
+		return Components.Num() ? Components.Last() : FString();
+	}
+
 	FString FURL_RFC3986::GetAuthority() const
 	{
 		FString Authority;
diff --git a/Engine/Plugins/Media/ElectraUtil/Source/ElectraBase/Public/MediaDecoderOutput.h b/Engine/Plugins/Media/ElectraUtil/Source/ElectraBase/Public/MediaDecoderOutput.h
--- a/Engine/Plugins/Media/ElectraUtil/Source/ElectraBase/Public/MediaDecoderOutput.h
+++ b/Engine/Plugins/Media/ElectraUtil/Source/ElectraBase/Public/MediaDecoderOutput.h
@@ -8,6 +8,9 @@
 #include "Containers/Queue.h"
 #include "Misc/ScopeLock.h"
 
+#include "ParameterDictionary.h"
+
+
 struct FDecoderTimeStamp
 {
 	FDecoderTimeStamp() {}
@@ -276,16 +279,24 @@ private:
 };
 
 
+class IDecoderOutput;
+
 class IDecoderOutputOwner
 {
 public:
-	virtual void SampleReleasedToPool(FTimespan durationToRelease) = 0;
+	virtual void SampleReleasedToPool(IDecoderOutput* InDecoderOutput) = 0;
 };
 
+
 class IDecoderOutput : public IDecoderOutputPoolable
 {
 public:
 	virtual void SetOwner(const TSharedPtr<IDecoderOutputOwner, ESPMode::ThreadSafe>& Renderer) = 0;
-};
-
+	virtual FDecoderTimeStamp GetTime() const = 0;
+	virtual FTimespan GetDuration() const = 0;
 
+	virtual Electra::FParamDict& GetMutablePropertyDictionary()
+	{ return PropertyDictionary; }
+private:
+	Electra::FParamDict PropertyDictionary;
+};
diff --git a/Engine/Plugins/Media/ElectraUtil/Source/ElectraBase/Public/MediaVideoDecoderOutput.h b/Engine/Plugins/Media/ElectraUtil/Source/ElectraBase/Public/MediaVideoDecoderOutput.h
--- a/Engine/Plugins/Media/ElectraUtil/Source/ElectraBase/Public/MediaVideoDecoderOutput.h
+++ b/Engine/Plugins/Media/ElectraUtil/Source/ElectraBase/Public/MediaVideoDecoderOutput.h
@@ -45,11 +45,15 @@ protected:
 class FVideoDecoderOutput : public IDecoderOutput
 {
 public:
-	virtual ~FVideoDecoderOutput() {}
+	virtual ~FVideoDecoderOutput() = default;
 
-	const Electra::FParamDict& GetDict() const { return *ParamDict; }
+	virtual const Electra::FParamDict& GetDict() const
+	{ 
+		check(ParamDict.IsValid());
+		return *ParamDict; 
+	}
 
-	FDecoderTimeStamp GetTime() const
+	virtual FDecoderTimeStamp GetTime() const
 	{
 		if (!ParamDict)
 		{
@@ -59,16 +63,16 @@ public:
 		return FDecoderTimeStamp(pts.GetAsTimespan(), pts.GetSequenceIndex());
 	}
 
-	FTimespan GetDuration() const
+	virtual FTimespan GetDuration() const
 	{
 		if (!ParamDict)
 		{
-			return 0;
+			return FTimespan(0);
 		}
 		return FTimespan(ParamDict->GetValue("duration").GetTimeValue().GetAsHNS());
 	}
 
-	FIntPoint GetOutputDim() const
+	virtual FIntPoint GetOutputDim() const
 	{
 		if (!ParamDict)
 		{
@@ -84,7 +88,7 @@ public:
 		return Cached.OutputDim;
 	}
 
-	double GetAspectRatio() const
+	virtual double GetAspectRatio() const
 	{
 		if (!ParamDict)
 		{
@@ -94,7 +98,7 @@ public:
 		return ((double)Dim.X / (double)Dim.Y) * ParamDict->GetValue("aspect_ratio").SafeGetDouble(1.0);
 	}
 
-	FVideoDecoderCropInfo GetCropInfo() const
+	virtual FVideoDecoderCropInfo GetCropInfo() const
 	{
 		if (!ParamDict)
 		{
@@ -127,7 +131,7 @@ public:
 		return Dim;
 	}
 
-	EPixelFormat GetFormat() const
+	virtual EPixelFormat GetFormat() const
 	{
 		if (!ParamDict)
 		{
@@ -136,7 +140,7 @@ public:
 		return (EPixelFormat)ParamDict->GetValue("pixelfmt").SafeGetInt64((int64)EPixelFormat::PF_Unknown);
 	}
 
-	EVideoOrientation GetOrientation() const
+	virtual EVideoOrientation GetOrientation() const
 	{
 		if (!ParamDict)
 		{
@@ -152,12 +156,7 @@ public:
 	}
 
 protected:
-	FVideoDecoderOutput()
-	{
-		Cached.Flags = 0;
-		Cached.OutputDim = FIntPoint::ZeroValue;
-		Cached.Orientation = EVideoOrientation::Original;
-	}
+	FVideoDecoderOutput() = default;
 
 	void Initialize(Electra::FParamDict* InParamDict)
 	{
@@ -166,10 +165,12 @@ protected:
 	}
 
 private:
-	TUniquePtr<Electra::FParamDict> ParamDict;
-
-	mutable struct FCached
+	struct FCached
 	{
+		FCached()
+		{
+			OutputDim = FIntPoint::ZeroValue;
+		}
 		enum
 		{
 			Valid_CropInfo = 1 << 0,
@@ -177,11 +178,13 @@ private:
 			Valid_Orientation = 1 << 2,
 		};
 
-		uint32 Flags;
 		FVideoDecoderCropInfo CropInfo;
 		FIntPoint OutputDim;
-		EVideoOrientation Orientation;
-	} Cached;
+		EVideoOrientation Orientation = EVideoOrientation::Original;;
+		uint32 Flags = 0;
+	};
+	mutable FCached Cached;
+	TUniquePtr<Electra::FParamDict> ParamDict;
 };
 
 using FVideoDecoderOutputPtr = TSharedPtr<FVideoDecoderOutput, ESPMode::ThreadSafe>;
diff --git a/Engine/Plugins/Media/ElectraUtil/Source/ElectraBase/Public/Utilities/TimeUtilities.h b/Engine/Plugins/Media/ElectraUtil/Source/ElectraBase/Public/Utilities/TimeUtilities.h
--- a/Engine/Plugins/Media/ElectraUtil/Source/ElectraBase/Public/Utilities/TimeUtilities.h
+++ b/Engine/Plugins/Media/ElectraUtil/Source/ElectraBase/Public/Utilities/TimeUtilities.h
@@ -15,7 +15,6 @@ namespace Electra
 		bool ELECTRABASE_API ParseDuration(FTimeValue& OutTimeValue, const TCHAR* InDuration);
 	}
 
-
 	namespace RFC7231
 	{
 		bool ELECTRABASE_API ParseDateTime(FTimeValue& OutTimeValue, const FString& DateTime);
@@ -31,5 +30,10 @@ namespace Electra
 		bool ELECTRABASE_API ParseFloatString(FTimeValue& OutTimeValue, const FString& Seconds);
 	}
 
+	namespace RFC5905
+	{
+		bool ELECTRABASE_API ParseNTPTime(FTimeValue& OutTimeValue, uint64 NtpTimestampFormat);
+	}
+
 } // namespace Electra
 
diff --git a/Engine/Plugins/Media/ElectraUtil/Source/ElectraBase/Public/Utilities/URLParser.h b/Engine/Plugins/Media/ElectraUtil/Source/ElectraBase/Public/Utilities/URLParser.h
--- a/Engine/Plugins/Media/ElectraUtil/Source/ElectraBase/Public/Utilities/URLParser.h
+++ b/Engine/Plugins/Media/ElectraUtil/Source/ElectraBase/Public/Utilities/URLParser.h
@@ -34,6 +34,8 @@ namespace Electra
 		bool IsAbsolute() const;
 		// Returns the path as individual components. Like GetPath() the components will have escape sequences already decoded.
 		void GetPathComponents(TArray<FString>& OutPathComponents) const;
+		// Returns the last path component (the "filename").
+		FString GetLastPathComponent() const;
 		// Resolves a relative URL against this one.
 		FURL_RFC3986& ResolveWith(const FString& InChildURL);
 		// Resolves this URL (which should be relative) against the specified URL.
diff --git a/Engine/Plugins/Media/ElectraUtil/Source/ElectraHTTPStream/Private/Apple/AppleElectraHTTPStream.cpp b/Engine/Plugins/Media/ElectraUtil/Source/ElectraHTTPStream/Private/Apple/AppleElectraHTTPStream.cpp
--- a/Engine/Plugins/Media/ElectraUtil/Source/ElectraHTTPStream/Private/Apple/AppleElectraHTTPStream.cpp
+++ b/Engine/Plugins/Media/ElectraUtil/Source/ElectraHTTPStream/Private/Apple/AppleElectraHTTPStream.cpp
@@ -67,16 +67,16 @@ public:
 	virtual ~FElectraHTTPStreamApple();
 
 	FElectraHTTPStreamApple();
-	virtual bool Initialize(const Electra::FParamDict& InOptions);
+	bool Initialize(const Electra::FParamDict& InOptions);
 
-	virtual FElectraHTTPStreamThreadHandlerDelegate& ThreadHandlerDelegate() override
+	FElectraHTTPStreamThreadHandlerDelegate& ThreadHandlerDelegate() override
 	{ return ThreadHandlerCallback; }
 
-	virtual void Close() override;
+	void Close() override;
 
-	virtual IElectraHTTPStreamRequestPtr CreateRequest() override;
+	IElectraHTTPStreamRequestPtr CreateRequest() override;
 
-	virtual void AddRequest(IElectraHTTPStreamRequestPtr Request) override;
+	void AddRequest(IElectraHTTPStreamRequestPtr Request) override;
 
 	NSURLSession* GetSession()
 	{ return Session; }
@@ -188,33 +188,35 @@ public:
 	FElectraHTTPStreamRequestApple();
 	virtual ~FElectraHTTPStreamRequestApple();
 
-	virtual void SetVerb(const FString& InVerb) override
+	void SetVerb(const FString& InVerb) override
 	{ Verb = InVerb; }
-	virtual IElectraHTTPStreamBuffer& POSTDataBuffer() override
+	void EnableTimingTraces() override
+	{ Response->SetEnableTimingTraces(); }
+	IElectraHTTPStreamBuffer& POSTDataBuffer() override
 	{ return PostData; }
-	virtual void SetUserAgent(const FString& InUserAgent) override
+	void SetUserAgent(const FString& InUserAgent) override
 	{ UserAgent = InUserAgent; }
-	virtual void SetURL(const FString& InURL) override
+	void SetURL(const FString& InURL) override
 	{ URL = InURL; }
-	virtual void SetRange(const FString& InRange) override
+	void SetRange(const FString& InRange) override
 	{ Range = InRange; }
-	virtual void AllowCompression(bool bInAllowCompression) override
+	void AllowCompression(bool bInAllowCompression) override
 	{ bAllowCompression = bInAllowCompression; }
-	virtual void AllowUnsafeRequestsForDebugging() override
+	void AllowUnsafeRequestsForDebugging() override
 	{
 	#if ELECTRA_HTTPSTREAM_APPLE_ALLOW_UNSAFE_CONNECTIONS_FOR_DEBUGGING
 		bAllowUnsafeConnectionsForDebugging = true;
 	#endif
 	}
-	virtual void AddHeader(const FString& Header, const FString& Value, bool bAppendIfExists) override;
-	virtual FElectraHTTPStreamNotificationDelegate& NotificationDelegate() override
+	void AddHeader(const FString& Header, const FString& Value, bool bAppendIfExists) override;
+	FElectraHTTPStreamNotificationDelegate& NotificationDelegate() override
 	{ return NotificationCallback; }
-	virtual void Cancel() override;
-	virtual IElectraHTTPStreamResponsePtr GetResponse() override
+	void Cancel() override;
+	IElectraHTTPStreamResponsePtr GetResponse() override
 	{ return Response; }
-	virtual bool HasFailed() override
+	bool HasFailed() override
 	{ return Response->GetErrorMessage().Len() > 0; }
-	virtual FString GetErrorMessage() override
+	FString GetErrorMessage() override
 	{ return Response->GetErrorMessage(); }
 
 	NSURLSessionTask* GetTaskHandle()
diff --git a/Engine/Plugins/Media/ElectraUtil/Source/ElectraHTTPStream/Private/Curl/CurlElectraHTTPStream.cpp b/Engine/Plugins/Media/ElectraUtil/Source/ElectraHTTPStream/Private/Curl/CurlElectraHTTPStream.cpp
--- a/Engine/Plugins/Media/ElectraUtil/Source/ElectraHTTPStream/Private/Curl/CurlElectraHTTPStream.cpp
+++ b/Engine/Plugins/Media/ElectraUtil/Source/ElectraHTTPStream/Private/Curl/CurlElectraHTTPStream.cpp
@@ -66,6 +66,11 @@ void FElectraHTTPStreamRequestLibCurl::SetVerb(const FString& InVerb)
 	Verb = InVerb; 
 }
 
+void FElectraHTTPStreamRequestLibCurl::EnableTimingTraces()
+{
+	Response->SetEnableTimingTraces();
+}
+
 IElectraHTTPStreamBuffer& FElectraHTTPStreamRequestLibCurl::POSTDataBuffer()
 { 
 	return PostData; 
diff --git a/Engine/Plugins/Media/ElectraUtil/Source/ElectraHTTPStream/Private/Curl/CurlElectraHTTPStreamRequest.h b/Engine/Plugins/Media/ElectraUtil/Source/ElectraHTTPStream/Private/Curl/CurlElectraHTTPStreamRequest.h
--- a/Engine/Plugins/Media/ElectraUtil/Source/ElectraHTTPStream/Private/Curl/CurlElectraHTTPStreamRequest.h
+++ b/Engine/Plugins/Media/ElectraUtil/Source/ElectraHTTPStream/Private/Curl/CurlElectraHTTPStreamRequest.h
@@ -47,19 +47,20 @@ public:
 	FElectraHTTPStreamRequestLibCurl();
 	virtual ~FElectraHTTPStreamRequestLibCurl();
 
-	virtual void SetVerb(const FString& InVerb) override;
-	virtual IElectraHTTPStreamBuffer& POSTDataBuffer() override;
-	virtual void SetUserAgent(const FString& InUserAgent) override;
-	virtual void SetURL(const FString& InURL) override;
-	virtual void SetRange(const FString& InRange) override;
-	virtual void AllowCompression(bool bInAllowCompression) override;
-	virtual void AllowUnsafeRequestsForDebugging() override;
-	virtual void AddHeader(const FString& Header, const FString& Value, bool bAppendIfExists) override;
-	virtual FElectraHTTPStreamNotificationDelegate& NotificationDelegate() override;
-	virtual void Cancel() override;
-	virtual IElectraHTTPStreamResponsePtr GetResponse() override;
-	virtual bool HasFailed() override;
-	virtual FString GetErrorMessage() override;
+	void SetVerb(const FString& InVerb) override;
+	void EnableTimingTraces() override;
+	IElectraHTTPStreamBuffer& POSTDataBuffer() override;
+	void SetUserAgent(const FString& InUserAgent) override;
+	void SetURL(const FString& InURL) override;
+	void SetRange(const FString& InRange) override;
+	void AllowCompression(bool bInAllowCompression) override;
+	void AllowUnsafeRequestsForDebugging() override;
+	void AddHeader(const FString& Header, const FString& Value, bool bAppendIfExists) override;
+	FElectraHTTPStreamNotificationDelegate& NotificationDelegate() override;
+	void Cancel() override;
+	IElectraHTTPStreamResponsePtr GetResponse() override;
+	bool HasFailed() override;
+	FString GetErrorMessage() override;
 
 	EState GetCurrentState();
 	bool WasCanceled();
diff --git a/Engine/Plugins/Media/ElectraUtil/Source/ElectraHTTPStream/Private/ElectraHTTPStreamResponse.h b/Engine/Plugins/Media/ElectraUtil/Source/ElectraHTTPStream/Private/ElectraHTTPStreamResponse.h
--- a/Engine/Plugins/Media/ElectraUtil/Source/ElectraHTTPStream/Private/ElectraHTTPStreamResponse.h
+++ b/Engine/Plugins/Media/ElectraUtil/Source/ElectraHTTPStream/Private/ElectraHTTPStreamResponse.h
@@ -18,74 +18,121 @@ public:
 	FElectraHTTPStreamResponse() = default;
 	virtual ~FElectraHTTPStreamResponse() = default;
 
-	virtual EStatus GetStatus() override
+	EStatus GetStatus() override
 	{ return CurrentStatus; }
 
-	virtual EState GetState() override
+	EState GetState() override
 	{ return CurrentState; }
 
-	virtual FString GetErrorMessage() override
+	FString GetErrorMessage() override
 	{ 
-		FScopeLock lock(&ErrorLock);
+		FScopeLock lock(&Lock);
 		return ErrorMessage; 
 	}
 
-	virtual int32 GetHTTPResponseCode() override
+	int32 GetHTTPResponseCode() override
 	{ return HTTPResponseCode; }
 
-	virtual int64 GetNumResponseBytesReceived() override
+	int64 GetNumResponseBytesReceived() override
 	{ return ResponseData.GetNumTotalBytesAdded(); }
 
-	virtual int64 GetNumRequestBytesSent() override
+	int64 GetNumRequestBytesSent() override
 	{ return 0;	}
 
-	virtual FString GetEffectiveURL() override
+	FString GetEffectiveURL() override
 	{ return EffectiveURL; }
 
-	virtual void GetAllHeaders(TArray<FElectraHTTPStreamHeader>& OutHeaders) override
+	void GetAllHeaders(TArray<FElectraHTTPStreamHeader>& OutHeaders) override
 	{ OutHeaders = ResponseHeaders;	}
 
-	virtual FString GetHTTPStatusLine() override
+	FString GetHTTPStatusLine() override
 	{ return HTTPStatusLine; }
-	virtual FString GetContentLengthHeader() override
+	FString GetContentLengthHeader() override
 	{ return ContentLength; }
-	virtual FString GetContentRangeHeader() override
+	FString GetContentRangeHeader() override
 	{ return ContentRange; }
-	virtual FString GetAcceptRangesHeader() override
+	FString GetAcceptRangesHeader() override
 	{ return AcceptRanges; }
-	virtual FString GetTransferEncodingHeader() override
+	FString GetTransferEncodingHeader() override
 	{ return TransferEncoding; }
-	virtual FString GetContentTypeHeader() override
+	FString GetContentTypeHeader() override
 	{ return ContentType; }
 
-	virtual IElectraHTTPStreamBuffer& GetResponseData() override
+	IElectraHTTPStreamBuffer& GetResponseData() override
 	{ return ResponseData; }
 
-	virtual double GetTimeElapsed() override
+	double GetTimeElapsed() override
 	{ return CurrentStatus == EStatus::NotRunning ? 0.0 : TimeUntilFinished > 0.0 ? TimeUntilFinished : FPlatformTime::Seconds() - StartTime; }
 
-	virtual double GetTimeSinceLastDataArrived() override
+	double GetTimeSinceLastDataArrived() override
 	{ return CurrentStatus == EStatus::NotRunning || TimeOfMostRecentReceive == 0.0 || TimeUntilFinished > 0.0 ? 0.0 : FPlatformTime::Seconds() - TimeOfMostRecentReceive; }
 
-	virtual double GetTimeUntilNameResolved() override
+	double GetTimeUntilNameResolved() override
 	{ return TimeUntilNameResolved; }
-	virtual double GetTimeUntilConnected() override
+	double GetTimeUntilConnected() override
 	{ return TimeUntilConnected; }
-	virtual double GetTimeUntilRequestSent() override
+	double GetTimeUntilRequestSent() override
 	{ return TimeUntilRequestSent; }
-	virtual double GetTimeUntilHeadersAvailable() override
+	double GetTimeUntilHeadersAvailable() override
 	{ return TimeUntilHeadersAvailable; }
-	virtual double GetTimeUntilFirstByte() override
+	double GetTimeUntilFirstByte() override
 	{ return TimeUntilFirstByte; }
-	virtual double GetTimeUntilFinished() override
+	double GetTimeUntilFinished() override
 	{ return TimeUntilFinished; }
 
-	void AddResponseData(const TArray<uint8>& InNewData)
-	{ ResponseData.AddData(InNewData); }
-	void AddResponseData(TArray<uint8>&& InNewData)
-	{ ResponseData.AddData(MoveTemp(InNewData)); }
-	void AddResponseData(const TConstArrayView<const uint8>& InNewData)
-	{ ResponseData.AddData(InNewData); }
+	int32 GetTimingTraces(TArray<FTimingTrace>* OutTraces, int32 InNumToRemove) override
+	{
+		FScopeLock lock(&Lock);
+		int32 Num = TimingTraces.Num();
+		if (OutTraces)
+		{
+			OutTraces->Append(TimingTraces);
+		}
+		if (InNumToRemove > 0)
+		{
+			int32 nr = InNumToRemove > Num ? Num : InNumToRemove;
+			TimingTraces.RemoveAt(0, nr);
+			Num = OutTraces ? Num : nr;	// If only removing the return value is to indicate how many got removed.
+		}
+		return Num;
+	}
+
+	void SetEnableTimingTraces() override
+	{
+		bCollectTimingTraces = true;
+	}
+
+
+	template<typename T>
+	void AddTrace(T InNewData, bool bAddTrace)
+	{
+		if (bAddTrace && bCollectTimingTraces)
+		{
+			FTimingTrace Trace;
+			Trace.TimeSinceStart = FPlatformTime::Seconds() - StartTime;
+			Trace.NumBytesAdded = InNewData.Num();
+			Trace.TotalBytesAdded = ResponseData.GetNumTotalBytesAdded() + InNewData.Num();
+			FScopeLock lock(&Lock);
+			TimingTraces.Emplace(MoveTemp(Trace));
+		}
+	}
+
+
+	void AddResponseData(const TArray<uint8>& InNewData, bool bAddTrace=true)
+	{ 
+		AddTrace(InNewData, bAddTrace);
+		ResponseData.AddData(InNewData); 
+	}
+	void AddResponseData(TArray<uint8>&& InNewData, bool bAddTrace=true)
+	{ 
+		AddTrace(InNewData, bAddTrace);
+		ResponseData.AddData(MoveTemp(InNewData)); 
+	}
+	void AddResponseData(const TConstArrayView<const uint8>& InNewData, bool bAddTrace=true)
+	{ 
+		AddTrace(InNewData, bAddTrace);
+		ResponseData.AddData(InNewData); 
+	}
 
 
 	void SetEOS()
@@ -93,7 +140,7 @@ public:
 
 	void SetErrorMessage(const FString& InErrorMessage)
 	{
-		FScopeLock lock(&ErrorLock);
+		FScopeLock lock(&Lock);
 		ErrorMessage = InErrorMessage;
 	}
 
@@ -123,7 +170,12 @@ public:
 	double TimeUntilFinished = 0.0;
 	double TimeOfMostRecentReceive = 0.0;
 
+	FCriticalSection Lock;
+
 	// Error message.
-	FCriticalSection ErrorLock;
 	FString ErrorMessage;
+
+	// Timing traces.
+	TArray<FTimingTrace> TimingTraces;
+	bool bCollectTimingTraces = false;
 };
diff --git a/Engine/Plugins/Media/ElectraUtil/Source/ElectraHTTPStream/Private/WinHttp/WinHttpElectraHTTPStream.cpp b/Engine/Plugins/Media/ElectraUtil/Source/ElectraHTTPStream/Private/WinHttp/WinHttpElectraHTTPStream.cpp
--- a/Engine/Plugins/Media/ElectraUtil/Source/ElectraHTTPStream/Private/WinHttp/WinHttpElectraHTTPStream.cpp
+++ b/Engine/Plugins/Media/ElectraUtil/Source/ElectraHTTPStream/Private/WinHttp/WinHttpElectraHTTPStream.cpp
@@ -54,6 +54,8 @@ THIRD_PARTY_INCLUDES_END
 #endif
 
 
+#define ELECTRA_SHARE_CONNECTION_HANDLE 1
+
 /***************************************************************************************************************************************************/
 /***************************************************************************************************************************************************/
 /***************************************************************************************************************************************************/
@@ -88,18 +90,21 @@ public:
 	static uint64 GetNewRequestIndex()
 	{ return ++NextRequestIndex; }
 
-	virtual FElectraHTTPStreamThreadHandlerDelegate& ThreadHandlerDelegate() override
+	FElectraHTTPStreamThreadHandlerDelegate& ThreadHandlerDelegate() override
 	{ return ThreadHandlerCallback; }
 
-	virtual void Close() override;
+	void Close() override;
 
-	virtual IElectraHTTPStreamRequestPtr CreateRequest() override;
+	IElectraHTTPStreamRequestPtr CreateRequest() override;
 
-	virtual void AddRequest(IElectraHTTPStreamRequestPtr Request) override;
+	void AddRequest(IElectraHTTPStreamRequestPtr Request) override;
 
 	HINTERNET GetSessionHandle()
 	{ return SessionHandle; }
 
+	HINTERNET GetConnectionHandle(const Electra::FURL_RFC3986& InForURL, INTERNET_PORT InPort, bool bForceNew);
+	void ReturnConnectionHandle(HINTERNET InConnectionHandle);
+
 	void TriggerWorkSignal()
 	{ HaveWorkSignal.Signal(); }
 
@@ -112,6 +117,15 @@ private:
 		TWeakPtr<FElectraHTTPStreamRequestWinHttp, ESPMode::ThreadSafe>	Request;
 	};
 
+	struct FSharedConnectionHandle
+	{
+		HINTERNET Handle = NULL;
+		FString Host;
+		INTERNET_PORT Port = 0;
+		int32 RefCount = 0;
+		double TimeWhenLastRefDropped = -1.0;
+	};
+
 	// Methods from FRunnable
 	virtual uint32 Run() override final;
 	virtual void Stop() override final;
@@ -149,10 +163,22 @@ private:
 	void SetupNewRequests();
 	void UpdateActiveRequests();
 	void HandleCompletedRequests();
+	void HandleIdleConnections();
+	void CloseAllConnectionHandles();
+
+	// Configuration
+#if ELECTRA_SHARE_CONNECTION_HANDLE
+	const bool bCloseOnLast = false;
+#else
+	const bool bCloseOnLast = true;
+#endif
+	const double CloseHandleAfterSecondsIdle = 10.0;
 
 
 	HINTERNET SessionHandle = NULL;
 
+	TArray<FSharedConnectionHandle> SharedConnectionHandles;
+
 	FThreadSafeCounter ExitRequest;
 	FRunnableThread* Thread = nullptr;
 	FTimeWaitableSignal HaveWorkSignal;
@@ -216,30 +242,33 @@ public:
 	FElectraHTTPStreamRequestWinHttp(uint64 InRequestIndex);
 	virtual ~FElectraHTTPStreamRequestWinHttp();
 
-	virtual void SetVerb(const FString& InVerb) override
+	void SetVerb(const FString& InVerb) override
 	{ Verb = InVerb; }
+	
+	void EnableTimingTraces() override
+	{ Response->SetEnableTimingTraces(); }
 
-	virtual IElectraHTTPStreamBuffer& POSTDataBuffer() override
+	IElectraHTTPStreamBuffer& POSTDataBuffer() override
 	{ return PostData; }
 
-	virtual void SetUserAgent(const FString& InUserAgent) override
+	void SetUserAgent(const FString& InUserAgent) override
 	{ UserAgent = InUserAgent; }
 
-	virtual void SetURL(const FString& InURL) override
+	void SetURL(const FString& InURL) override
 	{ URL = InURL; }
-	virtual void SetRange(const FString& InRange) override
+	void SetRange(const FString& InRange) override
 	{ Range = InRange; }
-	virtual void AllowCompression(bool bInAllowCompression) override
+	void AllowCompression(bool bInAllowCompression) override
 	{ bAllowCompression = bInAllowCompression; }
 
-	virtual void AllowUnsafeRequestsForDebugging() override
+	void AllowUnsafeRequestsForDebugging() override
 	{
 	#ifdef ELECTRA_HTTPSTREAM_WINHTTP_ALLOW_UNSAFE_CONNECTIONS_FOR_DEBUGGING
 		bAllowUnsafeConnectionsForDebugging = true;
 	#endif
 	}
 
-	virtual void AddHeader(const FString& Header, const FString& Value, bool bAppendIfExists) override
+	void AddHeader(const FString& Header, const FString& Value, bool bAppendIfExists) override
 	{
 		// Ignore a few headers that will be set automatically.
 		if (Header.Equals(TEXT("User-Agent"), ESearchCase::IgnoreCase) ||
@@ -273,23 +302,23 @@ public:
 		}
 	}
 
-	virtual FElectraHTTPStreamNotificationDelegate& NotificationDelegate() override
+	FElectraHTTPStreamNotificationDelegate& NotificationDelegate() override
 	{ return NotificationCallback; }
 
-	virtual void Cancel() override
+	void Cancel() override
 	{
 		FScopeLock lock(&NotificationLock);
 		NotificationCallback.Unbind();
 		bCancel = true;
 	}
 
-	virtual IElectraHTTPStreamResponsePtr GetResponse() override
+	IElectraHTTPStreamResponsePtr GetResponse() override
 	{ return Response; }
 
-	virtual bool HasFailed() override
+	bool HasFailed() override
 	{ return Response->GetErrorMessage().Len() > 0; }
 
-	virtual FString GetErrorMessage() override
+	FString GetErrorMessage() override
 	{ return Response->GetErrorMessage(); }
 
 	uint64 GetRequestIndex()
@@ -301,7 +330,7 @@ public:
 	bool WasCanceled()
 	{ return bCancel; }
 
-	bool Setup(FElectraHTTPStreamWinHttp* OwningManager);
+	bool Setup(TSharedPtr<FElectraHTTPStreamWinHttp, ESPMode::ThreadSafe> OwningManager);
 	bool Execute();
 	bool RequestResponse();
 	bool ParseHeaders(bool bNotify);
@@ -353,6 +382,8 @@ private:
 	FCriticalSection NotificationLock;
 	FElectraHTTPStreamNotificationDelegate NotificationCallback;
 
+	TWeakPtr<FElectraHTTPStreamWinHttp, ESPMode::ThreadSafe> OwningManager;
+
 	// Handles as required by WinHttp
 	HINTERNET ConnectionHandle = NULL;
 	HINTERNET RequestHandle = NULL;
@@ -387,7 +418,7 @@ FElectraHTTPStreamRequestWinHttp::~FElectraHTTPStreamRequestWinHttp()
 	Close();
 }
 
-bool FElectraHTTPStreamRequestWinHttp::Setup(FElectraHTTPStreamWinHttp* OwningManager)
+bool FElectraHTTPStreamRequestWinHttp::Setup(TSharedPtr<FElectraHTTPStreamWinHttp, ESPMode::ThreadSafe> InOwningManager)
 {
 	bool bIsRedirect = bMustRedirect;
 	if (bIsRedirect)
@@ -430,7 +461,7 @@ bool FElectraHTTPStreamRequestWinHttp::Setup(FElectraHTTPStreamWinHttp* OwningMa
 		LexFromString(p, *UrlPort);
 		Port = static_cast<INTERNET_PORT>(p);
 	}
-	HINTERNET ch = WinHttpConnect(OwningManager->GetSessionHandle(), TCHAR_TO_WCHAR(*UrlParser.GetHost()), Port, 0);
+	HINTERNET ch = InOwningManager->GetConnectionHandle(UrlParser, Port, false);
 	if (ch == NULL)
 	{
 		Response->SetErrorMessage(ElectraHTTPStreamWinHttp::GetErrorLogMessage(TEXT("WinHttpConnect()"), GetLastError()));
@@ -438,6 +469,7 @@ bool FElectraHTTPStreamRequestWinHttp::Setup(FElectraHTTPStreamWinHttp* OwningMa
 		return false;
 	}
 	ConnectionHandle = ch;
+	OwningManager = InOwningManager;
 
 	bool bIsSecure = UrlParser.GetScheme().Equals(TEXT("https"), ESearchCase::IgnoreCase);
 #if ELECTRA_HTTPSTREAM_REQUIRES_SECURE_CONNECTIONS
@@ -842,7 +874,15 @@ void FElectraHTTPStreamRequestWinHttp::Close(bool bHandlesOnly)
 	}
 	if (ConnectionHandle)
 	{
-		WinHttpCloseHandle(ConnectionHandle);
+		TSharedPtr<FElectraHTTPStreamWinHttp, ESPMode::ThreadSafe> Owner = OwningManager.Pin();
+		if (Owner.IsValid())
+		{
+			Owner->ReturnConnectionHandle(ConnectionHandle);
+		}
+		else
+		{
+			WinHttpCloseHandle(ConnectionHandle);
+		}
 		ConnectionHandle = NULL;
 	}
 	if (!bHandlesOnly)
@@ -1139,6 +1179,7 @@ void FElectraHTTPStreamWinHttp::Close()
 		Thread = nullptr;
 	}
 
+	CloseAllConnectionHandles();
 	if (SessionHandle)
 	{
 		WinHttpCloseHandle(SessionHandle);
@@ -1148,6 +1189,99 @@ void FElectraHTTPStreamWinHttp::Close()
 	RemoveOutdatedRequests();
 }
 
+
+HINTERNET FElectraHTTPStreamWinHttp::GetConnectionHandle(const Electra::FURL_RFC3986& InForURL, INTERNET_PORT InPort, bool bForceNew)
+{
+	FString Host = InForURL.GetHost();
+
+#if ELECTRA_SHARE_CONNECTION_HANDLE
+	if (!bForceNew)
+	{
+		for(int32 i=0; i<SharedConnectionHandles.Num(); ++i)
+		{
+			if (SharedConnectionHandles[i].Port == InPort && SharedConnectionHandles[i].Host.Equals(Host, ESearchCase::CaseSensitive))
+			{
+				++SharedConnectionHandles[i].RefCount;
+				SharedConnectionHandles[i].TimeWhenLastRefDropped = -1.0;
+				check(SharedConnectionHandles[i].Handle);
+				return SharedConnectionHandles[i].Handle;
+			}
+		}
+	}
+#endif
+
+	FSharedConnectionHandle sh;
+	HINTERNET Handle = WinHttpConnect(GetSessionHandle(), TCHAR_TO_WCHAR(*Host), InPort, 0);
+	if (Handle)
+	{
+		sh.Handle = Handle;
+		sh.Host = Host;
+		sh.Port = InPort;
+		sh.RefCount = 1;
+		SharedConnectionHandles.Emplace(MoveTemp(sh));
+	}
+	return Handle;
+}
+
+void FElectraHTTPStreamWinHttp::ReturnConnectionHandle(HINTERNET InConnectionHandle)
+{
+	if (InConnectionHandle)
+	{
+		for(int32 i=0; i<SharedConnectionHandles.Num(); ++i)
+		{
+			if (SharedConnectionHandles[i].Handle == InConnectionHandle)
+			{
+				check(SharedConnectionHandles[i].RefCount);
+
+				if (--SharedConnectionHandles[i].RefCount == 0)
+				{
+					if (bCloseOnLast)
+					{
+						WinHttpCloseHandle(SharedConnectionHandles[i].Handle);
+						SharedConnectionHandles[i].Handle = NULL;
+						SharedConnectionHandles.RemoveAtSwap(i);
+					}
+					else
+					{
+						SharedConnectionHandles[i].TimeWhenLastRefDropped = FPlatformTime::Seconds();
+					}
+				}
+				return;
+			}
+		}
+	}
+}
+
+void FElectraHTTPStreamWinHttp::CloseAllConnectionHandles()
+{
+	for(int32 i=0; i<SharedConnectionHandles.Num(); ++i)
+	{
+		WinHttpCloseHandle(SharedConnectionHandles[i].Handle);
+		SharedConnectionHandles[i].Handle = NULL;
+	}
+	SharedConnectionHandles.Reset();
+}
+
+
+void FElectraHTTPStreamWinHttp::HandleIdleConnections()
+{
+	double Now = FPlatformTime::Seconds() - CloseHandleAfterSecondsIdle;
+
+	for(int32 i=0; i<SharedConnectionHandles.Num(); ++i)
+	{
+		if (SharedConnectionHandles[i].RefCount == 0 && 
+			SharedConnectionHandles[i].TimeWhenLastRefDropped > 0.0 &&
+			Now > SharedConnectionHandles[i].TimeWhenLastRefDropped)
+		{
+			WinHttpCloseHandle(SharedConnectionHandles[i].Handle);
+			SharedConnectionHandles[i].Handle = NULL;
+			SharedConnectionHandles.RemoveAtSwap(i);
+			--i;
+		}
+	}
+}
+
+
 IElectraHTTPStreamRequestPtr FElectraHTTPStreamWinHttp::CreateRequest()
 {
 	LLM_SCOPE(ELLMTag::MediaStreaming);
@@ -1192,6 +1326,7 @@ uint32 FElectraHTTPStreamWinHttp::Run()
 		SetupNewRequests();
 		UpdateActiveRequests();
 		HandleCompletedRequests();
+		HandleIdleConnections();
 		}
 
 		// User callback
@@ -1242,7 +1377,7 @@ void FElectraHTTPStreamWinHttp::SetupNewRequests()
 	RequestLock.Unlock();
 	for(auto &Request : NewReqs)
 	{
-		if (Request->Setup(this))
+		if (Request->Setup(AsShared()))
 		{
 			FRequestPointers Pointers;
 			Pointers.Owner = AsShared();
diff --git a/Engine/Plugins/Media/ElectraUtil/Source/ElectraHTTPStream/Public/ElectraHTTPStream.h b/Engine/Plugins/Media/ElectraUtil/Source/ElectraHTTPStream/Public/ElectraHTTPStream.h
--- a/Engine/Plugins/Media/ElectraUtil/Source/ElectraHTTPStream/Public/ElectraHTTPStream.h
+++ b/Engine/Plugins/Media/ElectraUtil/Source/ElectraHTTPStream/Public/ElectraHTTPStream.h
@@ -93,6 +93,16 @@ public:
 		Finished
 	};
 
+	struct FTimingTrace
+	{
+		// Seconds since the request started.
+		double TimeSinceStart;
+		// Bytes added with the most recent chunk.
+		int64 NumBytesAdded;
+		// Total bytes added to the buffer including the most recent chunk.
+		int64 TotalBytesAdded;
+	};
+
 	virtual EStatus GetStatus() = 0;
 	virtual EState GetState() = 0;
 
@@ -125,6 +135,13 @@ public:
 	virtual double GetTimeUntilHeadersAvailable() = 0;
 	virtual double GetTimeUntilFirstByte() = 0;
 	virtual double GetTimeUntilFinished() = 0;
+
+	// Appends a copy of the timing traces to the provided array and optionally removes the first n elements.
+	// Pass a nullptr to only remove elements without copying them. This is useful to first get the traces
+	// and in a second call remove those that are no longer needed.
+	// Returns number appended or removed.
+	virtual int32 GetTimingTraces(TArray<FTimingTrace>* OutTraces, int32 InNumToRemove) = 0;
+	virtual void SetEnableTimingTraces() = 0;
 };
 
 
@@ -140,6 +157,9 @@ public:
 	// Set verb, GET, HEAD or POST. If not set GET is used by default.
 	virtual void SetVerb(const FString& Verb) = 0;
 
+	// Enables collection of timing traces. (disabled by default)
+	virtual void EnableTimingTraces() = 0;
+
 	/*
 		Returns the buffer holding the data to be POSTed to the server.
 		Make sure to include a suitable "Content-Type: xxx" header via AddHeader()
diff --git a/Engine/Source/Runtime/AVEncoder/Private/Decoders/Windows/VideoDecoderH264_Windows.cpp b/Engine/Source/Runtime/AVEncoder/Private/Decoders/Windows/VideoDecoderH264_Windows.cpp
--- a/Engine/Source/Runtime/AVEncoder/Private/Decoders/Windows/VideoDecoderH264_Windows.cpp
+++ b/Engine/Source/Runtime/AVEncoder/Private/Decoders/Windows/VideoDecoderH264_Windows.cpp
@@ -62,11 +62,13 @@ struct FParamDict
 
 }
 
+class IDecoderOutput;
+
 class IDecoderOutputOwner
 {
 public:
 	virtual ~IDecoderOutputOwner() = default;
-	virtual void SampleReleasedToPool(FTimespan durationToRelease) = 0;
+	virtual void SampleReleasedToPool(IDecoderOutput* InDecoderOutput) = 0;
 };
 
 class FNativeVideoDecoderOutput
