ElectraUtil: Added missing module to plugin; init curl if it wasn't before

diff --git a/Engine/Plugins/Media/ElectraPlayer/Source/ElectraPlayerRuntime/Private/ElectraPlayerRuntimeModule.cpp b/Engine/Plugins/Media/ElectraPlayer/Source/ElectraPlayerRuntime/Private/ElectraPlayerRuntimeModule.cpp
--- a/Engine/Plugins/Media/ElectraPlayer/Source/ElectraPlayerRuntime/Private/ElectraPlayerRuntimeModule.cpp
+++ b/Engine/Plugins/Media/ElectraPlayer/Source/ElectraPlayerRuntime/Private/ElectraPlayerRuntimeModule.cpp
@@ -19,6 +19,8 @@
 #include "ElectraPlayer.h"
 #include "ElectraPlayerPlatform.h"
 
+#include "IElectraHTTPStreamModule.h"
+
 DEFINE_LOG_CATEGORY(LogElectraPlayer);
 DEFINE_LOG_CATEGORY(LogElectraPlayerStats);
 
@@ -89,6 +91,9 @@ public:
 		IAudioDecoderAAC::FSystemConfiguration SysCfgAAC;
 		IAudioDecoderAAC::Startup(SysCfgAAC);
 
+		// The HTTP stream module is required.
+		FModuleManager::LoadModuleChecked<IElectraHTTPStreamModule>("ElectraHTTPStream");
+
 		bInitialized = true;
 	}
 
diff --git a/Engine/Plugins/Media/ElectraPlayer/Source/ElectraPlayerRuntime/Private/Runtime/HTTP/HTTPResponseCache.cpp b/Engine/Plugins/Media/ElectraPlayer/Source/ElectraPlayerRuntime/Private/Runtime/HTTP/HTTPResponseCache.cpp
--- a/Engine/Plugins/Media/ElectraPlayer/Source/ElectraPlayerRuntime/Private/Runtime/HTTP/HTTPResponseCache.cpp
+++ b/Engine/Plugins/Media/ElectraPlayer/Source/ElectraPlayerRuntime/Private/Runtime/HTTP/HTTPResponseCache.cpp
@@ -411,15 +411,16 @@ void FHTTPResponseCache::CacheEntity(TSharedPtrTS<FCacheItem> EntityToAdd)
 {
 	if (EntityToAdd.IsValid())
 	{
-		check(EntityToAdd->EffectiveURL.Len());
-
-		int64 SizeRequired = EntityToAdd->Response->GetResponseData().GetNumTotalBytesAdded();
-		EvictToAddSize(SizeRequired);
-		FScopeLock ScopeLock(&Lock);
-		if (SizeInUse + SizeRequired <= MaxElementSize && Cache.Num() < MaxNumElements)
+		if (EntityToAdd->EffectiveURL.Len())
 		{
-			Cache.Insert(MoveTemp(EntityToAdd), 0);
-			SizeInUse += SizeRequired;
+			int64 SizeRequired = EntityToAdd->Response->GetResponseData().GetNumTotalBytesAdded();
+			EvictToAddSize(SizeRequired);
+			FScopeLock ScopeLock(&Lock);
+			if (SizeInUse + SizeRequired <= MaxElementSize && Cache.Num() < MaxNumElements)
+			{
+				Cache.Insert(MoveTemp(EntityToAdd), 0);
+				SizeInUse += SizeRequired;
+			}
 		}
 	}
 }
diff --git a/Engine/Plugins/Media/ElectraUtil/Source/ElectraHTTPStream/Private/Curl/CurlElectraHTTPStream.cpp b/Engine/Plugins/Media/ElectraUtil/Source/ElectraHTTPStream/Private/Curl/CurlElectraHTTPStream.cpp
--- a/Engine/Plugins/Media/ElectraUtil/Source/ElectraHTTPStream/Private/Curl/CurlElectraHTTPStream.cpp
+++ b/Engine/Plugins/Media/ElectraUtil/Source/ElectraHTTPStream/Private/Curl/CurlElectraHTTPStream.cpp
@@ -497,7 +497,7 @@ void FElectraHTTPStreamRequestLibCurl::SetupProxy(const FString& InProxy)
 bool FElectraHTTPStreamRequestLibCurl::Execute(FElectraHTTPStreamLibCurl* OwningManager)
 {
 	// Set the origin URL as effective URL first in case there are problems or no redirections.
-	Response->EffectiveURL = URL;
+	EffectiveURL = Response->EffectiveURL = URL;
 	Response->StartTime = FPlatformTime::Seconds();
 	Response->CurrentStatus = IElectraHTTPStreamResponse::EStatus::Running;
 
@@ -686,7 +686,10 @@ size_t FElectraHTTPStreamRequestLibCurl::CurlResponseHeaderCallback(void* InData
 				char* EffUrl = nullptr;
 				if (curl_easy_getinfo(CurlEasyHandle, CURLINFO_EFFECTIVE_URL, &EffUrl) == CURLE_OK)
 				{
-					EffectiveURL = FString(EffUrl);
+					if (EffUrl && *EffUrl)
+					{
+						EffectiveURL = FString(EffUrl);
+					}
 				}
 
 				// This is not necessarily the last header. There could be redirections following or
@@ -1270,6 +1273,69 @@ void FElectraHTTPStreamLibCurl::HandleCompletedRequests()
 #include "Modules/ModuleManager.h"
 #include "HttpModule.h"
 
+namespace ElectraHTTPStreamLibCurl
+{
+	/*
+		These memory hooks are identical to those in CurlHttp.h
+		and are used if we need Curl but the HTTP module does not.
+	*/
+	void* CurlMalloc(size_t Size)
+	{ 
+		return FMemory::Malloc(Size); 
+	}
+
+	void CurlFree(void* Ptr)
+	{ 
+		FMemory::Free(Ptr); 
+	}
+
+	void* CurlRealloc(void* Ptr, size_t Size)
+	{
+		void* Return = NULL;
+		if (Size)
+		{
+			Return = FMemory::Realloc(Ptr, Size);
+		}
+		return Return;
+	}
+
+	char* CurlStrdup(const char* ZeroTerminatedString)
+	{
+		char * Copy = NULL;
+		check(ZeroTerminatedString);
+		if (ZeroTerminatedString)
+		{
+			SIZE_T StrLen = FCStringAnsi::Strlen(ZeroTerminatedString);
+			Copy = reinterpret_cast<char*>(FMemory::Malloc(StrLen + 1));
+			if (Copy)
+			{
+				FCStringAnsi::Strcpy(Copy, StrLen, ZeroTerminatedString);
+				check(FCStringAnsi::Strcmp(Copy, ZeroTerminatedString) == 0);
+			}
+		}
+		return Copy;
+	}
+
+	void* CurlCalloc(size_t NumElems, size_t ElemSize)
+	{
+		void* Return = NULL;
+		const size_t Size = NumElems * ElemSize;
+		if (Size)
+		{
+			Return = FMemory::Malloc(Size);
+			if (Return)
+			{
+				FMemory::Memzero(Return, Size);
+			}
+		}
+		return Return;
+	}
+
+	// Flag indicating whether or not we initialized Curl.
+	static bool gDidInitCurl = false;
+}
+
+
 void FPlatformElectraHTTPStreamLibCurl::Startup()
 {
 	// Because of global methods like curl_global_init_mem() and similar methods for SSL have to rely on
@@ -1280,15 +1346,41 @@ void FPlatformElectraHTTPStreamLibCurl::Startup()
 
 	FModuleManager::LoadModuleChecked<FHttpModule>("HTTP");
 
+	int32 CurlInitFlags = CURL_GLOBAL_ALL;
+	(void)CurlInitFlags;
+
 	// Do the same for SSL. Here we get a hold of it and need to close it down in Shutdown().
 #if WITH_SSL
 	FSslModule& SslModule = FModuleManager::LoadModuleChecked<FSslModule>("SSL");
-	SslModule.GetSslManager().InitializeSsl();
+	if (SslModule.GetSslManager().InitializeSsl())
+	{
+		CurlInitFlags = CurlInitFlags & ~(CURL_GLOBAL_SSL);
+	}
+#endif
+
+#if ELECTRA_HTTPSTREAM_CURL_DO_GLOBAL_INIT
+	CURLcode InitResult = curl_global_init_mem(CurlInitFlags, ElectraHTTPStreamLibCurl::CurlMalloc, ElectraHTTPStreamLibCurl::CurlFree, ElectraHTTPStreamLibCurl::CurlRealloc, ElectraHTTPStreamLibCurl::CurlStrdup, ElectraHTTPStreamLibCurl::CurlCalloc);
+	/*
+		If already initialized this may return an error. If already initialized it will not overwrite the memory hooks.
+		Curl documentation says to call curl_global_cleanup() for each call to curl_global_init().
+		It doesn't say to do this only if the init was successful, but the source gives that away.
+	*/
+	if (InitResult == 0)
+	{
+		ElectraHTTPStreamLibCurl::gDidInitCurl = true;
+	}
 #endif
 }
 
 void FPlatformElectraHTTPStreamLibCurl::Shutdown()
 {
+#if ELECTRA_HTTPSTREAM_CURL_DO_GLOBAL_INIT
+	if (ElectraHTTPStreamLibCurl::gDidInitCurl)
+	{
+		ElectraHTTPStreamLibCurl::gDidInitCurl = false;
+		curl_global_cleanup();
+	}
+#endif
 #if WITH_SSL
 	FSslModule& SslModule = FModuleManager::LoadModuleChecked<FSslModule>("SSL");
 	SslModule.GetSslManager().ShutdownSsl();
diff --git a/Engine/Plugins/Media/ElectraUtil/Source/ElectraHTTPStream/Private/Curl/CurlElectraHTTPStreamConfig.h b/Engine/Plugins/Media/ElectraUtil/Source/ElectraHTTPStream/Private/Curl/CurlElectraHTTPStreamConfig.h
--- a/Engine/Plugins/Media/ElectraUtil/Source/ElectraHTTPStream/Private/Curl/CurlElectraHTTPStreamConfig.h
+++ b/Engine/Plugins/Media/ElectraUtil/Source/ElectraHTTPStream/Private/Curl/CurlElectraHTTPStreamConfig.h
@@ -9,11 +9,13 @@
 #if ELECTRA_HTTPSTREAM_XCURL == 0
 #define ELECTRA_HTTPSTREAM_CURL_USE_SHARE 1
 #define ELECTRA_HTTPSTREAM_CURL_POST_VIA_INFILESIZE 0
+#define ELECTRA_HTTPSTREAM_CURL_DO_GLOBAL_INIT 0
 #else
 #define ELECTRA_HTTPSTREAM_CURL_USE_SHARE 0
 #define ELECTRA_HTTPSTREAM_CURL_POST_VIA_INFILESIZE 1
 #define ELECTRA_HTTPSTREAM_CURL_USE_XFERINFO 1
 #define ELECTRA_HTTPSTREAM_CURL_USE_MULTIPOLL 1
+#define ELECTRA_HTTPSTREAM_CURL_DO_GLOBAL_INIT 1
 #endif
 
 #ifndef LIBCURL_VERSION_NUM
