ElectraHTTPStream: Added availability checks around proxy configuration

#rnx

diff --git a/Engine/Plugins/Media/ElectraUtil/Source/ElectraHTTPStream/Private/WinHttp/WinHttpElectraHTTPStream.cpp b/Engine/Plugins/Media/ElectraUtil/Source/ElectraHTTPStream/Private/WinHttp/WinHttpElectraHTTPStream.cpp
--- a/Engine/Plugins/Media/ElectraUtil/Source/ElectraHTTPStream/Private/WinHttp/WinHttpElectraHTTPStream.cpp
+++ b/Engine/Plugins/Media/ElectraUtil/Source/ElectraHTTPStream/Private/WinHttp/WinHttpElectraHTTPStream.cpp
@@ -1092,26 +1092,40 @@ bool FElectraHTTPStreamWinHttp::Initialize(const Electra::FParamDict& InOptions)
 	// Create the WinHTTP session handle.
 	DWORD SessionFlags = WINHTTP_FLAG_ASYNC;
 
+#if PLATFORM_WINDOWS
+	const bool bIsWindows7OrGreater = FPlatformMisc::VerifyWindowsVersion(6, 1);
+	const bool bIsWindows8Point1OrGreater = FPlatformMisc::VerifyWindowsVersion(6, 3);
+#endif
+
 #if ELECTRA_HTTPSTREAM_REQUIRES_SECURE_CONNECTIONS
 	SessionFlags |= WINHTTP_FLAG_SECURE_DEFAULTS;
 #endif
 
+	DWORD AccessType = WINHTTP_ACCESS_TYPE_AUTOMATIC_PROXY;
+#if PLATFORM_WINDOWS
+	if (!bIsWindows8Point1OrGreater)
+	{
+		AccessType = WINHTTP_ACCESS_TYPE_DEFAULT_PROXY;
+	}
+#endif
+
 	HINTERNET sh;
 	if (InOptions.HaveKey(TEXT("proxy")))
 	{
 		FString ProxyNameAndPort = InOptions.GetValue(TEXT("proxy")).SafeGetFString();
 		if (ProxyNameAndPort.Len())
 		{
-			sh = WinHttpOpen(NULL, WINHTTP_ACCESS_TYPE_NAMED_PROXY, TCHAR_TO_WCHAR(*ProxyNameAndPort), WINHTTP_NO_PROXY_BYPASS, SessionFlags);
+			AccessType = WINHTTP_ACCESS_TYPE_NAMED_PROXY;
+			sh = WinHttpOpen(NULL, AccessType, TCHAR_TO_WCHAR(*ProxyNameAndPort), WINHTTP_NO_PROXY_BYPASS, SessionFlags);
 		}
 		else
 		{
-			sh = WinHttpOpen(NULL, WINHTTP_ACCESS_TYPE_AUTOMATIC_PROXY, WINHTTP_NO_PROXY_NAME, WINHTTP_NO_PROXY_BYPASS, SessionFlags);
+			sh = WinHttpOpen(NULL, AccessType, WINHTTP_NO_PROXY_NAME, WINHTTP_NO_PROXY_BYPASS, SessionFlags);
 		}
 	}
 	else
 	{
-		sh = WinHttpOpen(NULL, WINHTTP_ACCESS_TYPE_AUTOMATIC_PROXY, WINHTTP_NO_PROXY_NAME, WINHTTP_NO_PROXY_BYPASS, SessionFlags);
+		sh = WinHttpOpen(NULL, AccessType, WINHTTP_NO_PROXY_NAME, WINHTTP_NO_PROXY_BYPASS, SessionFlags);
 	}
 	if (sh == NULL)
 	{
@@ -1128,7 +1142,6 @@ bool FElectraHTTPStreamWinHttp::Initialize(const Electra::FParamDict& InOptions)
 	DWORD SecureProtocols = 0;
 	#ifdef WINHTTP_FLAG_SECURE_PROTOCOL_TLS1_1
 		#if PLATFORM_WINDOWS
-			const bool bIsWindows7OrGreater = FPlatformMisc::VerifyWindowsVersion(6, 1);
 			if (bIsWindows7OrGreater)
 			{
 				SecureProtocols |= WINHTTP_FLAG_SECURE_PROTOCOL_TLS1_1;
@@ -1140,7 +1153,6 @@ bool FElectraHTTPStreamWinHttp::Initialize(const Electra::FParamDict& InOptions)
 	
 	#ifdef WINHTTP_FLAG_SECURE_PROTOCOL_TLS1_2
 		#if PLATFORM_WINDOWS
-			const bool bIsWindows8Point1OrGreater = FPlatformMisc::VerifyWindowsVersion(6, 3);
 			if (bIsWindows8Point1OrGreater)
 			{
 				SecureProtocols |= WINHTTP_FLAG_SECURE_PROTOCOL_TLS1_2;
