ElectraHTTPStream: Added an additional guard for the cases where joining a thread sometimes fails and is no indication that the thread has successfully ended.

#rnx

diff --git a/Engine/Plugins/Media/ElectraUtil/Source/ElectraHTTPStream/Private/Curl/CurlElectraHTTPStream.cpp b/Engine/Plugins/Media/ElectraUtil/Source/ElectraHTTPStream/Private/Curl/CurlElectraHTTPStream.cpp
--- a/Engine/Plugins/Media/ElectraUtil/Source/ElectraHTTPStream/Private/Curl/CurlElectraHTTPStream.cpp
+++ b/Engine/Plugins/Media/ElectraUtil/Source/ElectraHTTPStream/Private/Curl/CurlElectraHTTPStream.cpp
@@ -961,6 +961,8 @@ void FElectraHTTPStreamLibCurl::Close()
 	if (Thread)
 	{
 		Thread->Kill(true);
+		ThreadFinished.Lock();
+		ThreadFinished.Unlock();
 		delete Thread;
 		Thread = nullptr;
 	}
@@ -1109,6 +1111,8 @@ uint32 FElectraHTTPStreamLibCurl::Run()
 {
 	LLM_SCOPE(ELLMTag::MediaStreaming);
 
+	ThreadFinished.Lock();
+
 #if ELECTRA_HTTPSTREAM_CURL_USE_MULTIPOLL
 	WorkMultiPoll();
 #else
@@ -1131,6 +1135,7 @@ uint32 FElectraHTTPStreamLibCurl::Run()
 	}
 	RequestLock.Unlock();
 	HandleCompletedRequests();
+	ThreadFinished.Unlock();
 
 	return 0;
 }
diff --git a/Engine/Plugins/Media/ElectraUtil/Source/ElectraHTTPStream/Private/Curl/CurlElectraHTTPStreamImpl.h b/Engine/Plugins/Media/ElectraUtil/Source/ElectraHTTPStream/Private/Curl/CurlElectraHTTPStreamImpl.h
--- a/Engine/Plugins/Media/ElectraUtil/Source/ElectraHTTPStream/Private/Curl/CurlElectraHTTPStreamImpl.h
+++ b/Engine/Plugins/Media/ElectraUtil/Source/ElectraHTTPStream/Private/Curl/CurlElectraHTTPStreamImpl.h
@@ -84,6 +84,7 @@ private:
 	FThreadSafeCounter ExitRequest;
 	FRunnableThread* Thread = nullptr;
 	FTimeWaitableSignal HaveWorkSignal;
+	FCriticalSection ThreadFinished;
 
 	FElectraHTTPStreamThreadHandlerDelegate ThreadHandlerCallback;
 
