ElectraHTTPStream: Added optional platform connection limit

#rnx

diff --git a/Engine/Plugins/Media/ElectraUtil/Source/ElectraHTTPStream/Private/Curl/CurlElectraHTTPStream.cpp b/Engine/Plugins/Media/ElectraUtil/Source/ElectraHTTPStream/Private/Curl/CurlElectraHTTPStream.cpp
--- a/Engine/Plugins/Media/ElectraUtil/Source/ElectraHTTPStream/Private/Curl/CurlElectraHTTPStream.cpp
+++ b/Engine/Plugins/Media/ElectraUtil/Source/ElectraHTTPStream/Private/Curl/CurlElectraHTTPStream.cpp
@@ -395,6 +395,10 @@ bool FElectraHTTPStreamRequestLibCurl::Setup(FElectraHTTPStreamLibCurl* OwningMa
 	{
 		Result = curl_easy_setopt(CurlEasyHandle, CURLOPT_CONNECTTIMEOUT_MS, (long)ConnectionTimeoutMillis);
 	}
+#ifdef ELECTRA_HTTPSTREAM_CURL_PLATFORM_FORBID_REUSE
+	Result = curl_easy_setopt(CurlEasyHandle, CURLOPT_FORBID_REUSE, 1L);
+#endif
+
 
 	// HTTP/2 ?
 	if (OwningManager->SupportsHTTP2())
@@ -935,6 +939,26 @@ bool FElectraHTTPStreamLibCurl::Initialize(const Electra::FParamDict& InOptions)
 		return false;
 	}
 
+#ifdef ELECTRA_HTTPSTREAM_CURL_PLATFORM_LIMIT_MAX_TOTAL_CONNECTIONS
+	// Limit the maximum simultaneously open connectinos?
+	#if ELECTRA_HTTPSTREAM_CURL_PLATFORM_LIMIT_MAX_TOTAL_CONNECTIONS > 0
+	{
+		// Try setting the option. If it does not work it is ok.
+		mcode = curl_multi_setopt(CurlMultiHandle, CURLMOPT_MAX_TOTAL_CONNECTIONS, (long)ELECTRA_HTTPSTREAM_CURL_PLATFORM_LIMIT_MAX_TOTAL_CONNECTIONS);
+	}
+	#endif
+#endif
+
+#ifdef ELECTRA_HTTPSTREAM_CURL_PLATFORM_LIMIT_MAXCONNECTS
+	// Limit the maximum simultaneously open connections licburl may keep in its connection cache?
+	#if ELECTRA_HTTPSTREAM_CURL_PLATFORM_LIMIT_MAXCONNECTS > 0
+	{
+		// Try setting the option. If it does not work it is ok.
+		mcode = curl_multi_setopt(CurlMultiHandle, CURLMOPT_MAXCONNECTS, (long)ELECTRA_HTTPSTREAM_CURL_PLATFORM_LIMIT_MAXCONNECTS);
+	}
+	#endif
+#endif
+
 	if (bSupportsHTTP2)
 	{
 	#ifdef CURLPIPE_MULTIPLEX
diff --git a/Engine/Plugins/Media/ElectraUtil/Source/ElectraHTTPStream/Private/Curl/CurlElectraHTTPStreamConfig.h b/Engine/Plugins/Media/ElectraUtil/Source/ElectraHTTPStream/Private/Curl/CurlElectraHTTPStreamConfig.h
--- a/Engine/Plugins/Media/ElectraUtil/Source/ElectraHTTPStream/Private/Curl/CurlElectraHTTPStreamConfig.h
+++ b/Engine/Plugins/Media/ElectraUtil/Source/ElectraHTTPStream/Private/Curl/CurlElectraHTTPStreamConfig.h
@@ -7,11 +7,19 @@
 #include "Curl/CurlElectra.h"
 
 #if ELECTRA_HTTPSTREAM_XCURL == 0
+#ifdef ELECTRA_HTTPSTREAM_CURL_PLATFORM_USE_SHARE
+#define ELECTRA_HTTPSTREAM_CURL_USE_SHARE ELECTRA_HTTPSTREAM_CURL_PLATFORM_USE_SHARE
+#else
 #define ELECTRA_HTTPSTREAM_CURL_USE_SHARE 1
+#endif
 #define ELECTRA_HTTPSTREAM_CURL_POST_VIA_INFILESIZE 0
 #define ELECTRA_HTTPSTREAM_CURL_DO_GLOBAL_INIT 0
 #else
+#ifdef ELECTRA_HTTPSTREAM_CURL_PLATFORM_USE_SHARE
+#define ELECTRA_HTTPSTREAM_CURL_USE_SHARE ELECTRA_HTTPSTREAM_CURL_PLATFORM_USE_SHARE
+#else
 #define ELECTRA_HTTPSTREAM_CURL_USE_SHARE 0
+#endif
 #define ELECTRA_HTTPSTREAM_CURL_POST_VIA_INFILESIZE 1
 #define ELECTRA_HTTPSTREAM_CURL_USE_XFERINFO 1
 #define ELECTRA_HTTPSTREAM_CURL_USE_MULTIPOLL 1
