ElectraHTTPStream: Temporary workaround for an issue of aborted connections still counting as active ones.

#rnx

diff --git a/Engine/Plugins/Media/ElectraUtil/Source/ElectraHTTPStream/Private/Curl/CurlElectraHTTPStream.cpp b/Engine/Plugins/Media/ElectraUtil/Source/ElectraHTTPStream/Private/Curl/CurlElectraHTTPStream.cpp
--- a/Engine/Plugins/Media/ElectraUtil/Source/ElectraHTTPStream/Private/Curl/CurlElectraHTTPStream.cpp
+++ b/Engine/Plugins/Media/ElectraUtil/Source/ElectraHTTPStream/Private/Curl/CurlElectraHTTPStream.cpp
@@ -26,6 +26,8 @@
 #include "Curl/CurlElectraHTTPStreamRequest.h"
 #include "Curl/CurlLogError.h"
 
+#include "PlatformElectraHTTPStream.h"
+
 #if WITH_SSL
 #include "Ssl.h"
 #include <openssl/ssl.h>
@@ -1385,6 +1387,10 @@ void FPlatformElectraHTTPStreamLibCurl::Startup()
 #endif
 
 #if ELECTRA_HTTPSTREAM_CURL_DO_GLOBAL_INIT
+	#ifdef ELECTRA_HTTPSTREAM_CURL_PLATFORM_PRIVATE_INIT
+		ELECTRA_HTTPSTREAM_CURL_PLATFORM_PRIVATE_INIT::PrivateStartup();
+	#endif
+
 	CURLcode InitResult = curl_global_init_mem(CurlInitFlags, ElectraHTTPStreamLibCurl::CurlMalloc, ElectraHTTPStreamLibCurl::CurlFree, ElectraHTTPStreamLibCurl::CurlRealloc, ElectraHTTPStreamLibCurl::CurlStrdup, ElectraHTTPStreamLibCurl::CurlCalloc);
 	/*
 		If already initialized this may return an error. If already initialized it will not overwrite the memory hooks.
