Make sure the certificate chain used to evaluate trust is checked for nil when copied

diff --git a/Engine/Plugins/Media/ElectraUtil/Source/ElectraHTTPStream/Private/Apple/AppleElectraHTTPStream.cpp b/Engine/Plugins/Media/ElectraUtil/Source/ElectraHTTPStream/Private/Apple/AppleElectraHTTPStream.cpp
--- a/Engine/Plugins/Media/ElectraUtil/Source/ElectraHTTPStream/Private/Apple/AppleElectraHTTPStream.cpp
+++ b/Engine/Plugins/Media/ElectraUtil/Source/ElectraHTTPStream/Private/Apple/AppleElectraHTTPStream.cpp
@@ -1307,8 +1307,14 @@ FElectraHTTPStreamApple::FChallengeResponse FElectraHTTPStreamApple::ReceivedCha
 			// the chain starts with the server's cert itself, so walk backwards to optimize for roots first
 			TArray<TArray<uint8, TFixedAllocator<ISslCertificateManager::PUBLIC_KEY_DIGEST_SIZE>>> CertDigests;
 
-                        CFArrayRef Certificates = SecTrustCopyCertificateChain(RemoteTrust);
-                        CFIndex CertificateCount = CFArrayGetCount(Certificates);
+            CFArrayRef Certificates = SecTrustCopyCertificateChain(RemoteTrust);
+            if (Certificates == nil)
+            {
+                UE_LOG(LogHttp, Error, TEXT("No certificate could be copied in the certificate chain used to evaluate trust."));
+                return;
+            }
+            
+            CFIndex CertificateCount = CFArrayGetCount(Certificates);
 			for (int i = 0; CertificateCount; ++i)
 			{
 				SecCertificateRef Cert = SecTrustGetCertificateAtIndex(RemoteTrust, i);
diff --git a/Engine/Source/Runtime/Online/HTTP/Private/Apple/AppleHTTP.cpp b/Engine/Source/Runtime/Online/HTTP/Private/Apple/AppleHTTP.cpp
--- a/Engine/Source/Runtime/Online/HTTP/Private/Apple/AppleHTTP.cpp
+++ b/Engine/Source/Runtime/Online/HTTP/Private/Apple/AppleHTTP.cpp
@@ -700,8 +700,13 @@ static const unsigned char ecdsaSecp384r1Asn1Header[] =
             TArray<TArray<uint8, TFixedAllocator<ISslCertificateManager::PUBLIC_KEY_DIGEST_SIZE>>> CertDigests;
             
             CFArrayRef Certificates = SecTrustCopyCertificateChain(RemoteTrust);
+            if (Certificates == nil)
+            {
+                UE_LOG(LogHttp, Error, TEXT("No certificate could be copied in the certificate chain used to evaluate trust."));
+                return;
+            }
+            
             CFIndex CertificateCount = CFArrayGetCount(Certificates);
-
             for (int i = 0; CertificateCount; ++i)
             {
                 SecCertificateRef Cert = (SecCertificateRef)CFArrayGetValueAtIndex(Certificates, i);
