Horde fix for CL21965592
- update UE_LOG identifier and return statement

#rnx

diff --git a/Engine/Plugins/Media/ElectraUtil/Source/ElectraHTTPStream/Private/Apple/AppleElectraHTTPStream.cpp b/Engine/Plugins/Media/ElectraUtil/Source/ElectraHTTPStream/Private/Apple/AppleElectraHTTPStream.cpp
--- a/Engine/Plugins/Media/ElectraUtil/Source/ElectraHTTPStream/Private/Apple/AppleElectraHTTPStream.cpp
+++ b/Engine/Plugins/Media/ElectraUtil/Source/ElectraHTTPStream/Private/Apple/AppleElectraHTTPStream.cpp
@@ -1306,15 +1306,15 @@ FElectraHTTPStreamApple::FChallengeResponse FElectraHTTPStreamApple::ReceivedCha
 			// look at all certs in the remote chain and calculate the SHA256 hash of their DER-encoded SPKI
 			// the chain starts with the server's cert itself, so walk backwards to optimize for roots first
 			TArray<TArray<uint8, TFixedAllocator<ISslCertificateManager::PUBLIC_KEY_DIGEST_SIZE>>> CertDigests;
-
-            CFArrayRef Certificates = SecTrustCopyCertificateChain(RemoteTrust);
-            if (Certificates == nil)
-            {
-                UE_LOG(LogHttp, Error, TEXT("No certificate could be copied in the certificate chain used to evaluate trust."));
-                return;
-            }
-            
-            CFIndex CertificateCount = CFArrayGetCount(Certificates);
+			
+			CFArrayRef Certificates = SecTrustCopyCertificateChain(RemoteTrust);
+			if (Certificates == nil)
+			{
+				UE_LOG(LogElectraHTTPStream, Error, TEXT("No certificate could be copied in the certificate chain used to evaluate trust."));
+				return Response;
+			}
+			
+			CFIndex CertificateCount = CFArrayGetCount(Certificates);
 			for (int i = 0; CertificateCount; ++i)
 			{
 				SecCertificateRef Cert = SecTrustGetCertificateAtIndex(RemoteTrust, i);
