ElectraHTTPStream: Adjusted default connection values

#rnx

diff --git a/Engine/Plugins/Media/ElectraUtil/Source/ElectraHTTPStream/Private/Curl/CurlElectraHTTPStream.cpp b/Engine/Plugins/Media/ElectraUtil/Source/ElectraHTTPStream/Private/Curl/CurlElectraHTTPStream.cpp
--- a/Engine/Plugins/Media/ElectraUtil/Source/ElectraHTTPStream/Private/Curl/CurlElectraHTTPStream.cpp
+++ b/Engine/Plugins/Media/ElectraUtil/Source/ElectraHTTPStream/Private/Curl/CurlElectraHTTPStream.cpp
@@ -953,7 +953,7 @@ bool FElectraHTTPStreamLibCurl::Initialize(const Electra::FParamDict& InOptions)
 
 #ifdef ELECTRA_HTTPSTREAM_CURL_PLATFORM_LIMIT_MAXCONNECTS
 	// Limit the maximum simultaneously open connections licburl may keep in its connection cache?
-	#if ELECTRA_HTTPSTREAM_CURL_PLATFORM_LIMIT_MAXCONNECTS > 0
+	#if ELECTRA_HTTPSTREAM_CURL_PLATFORM_LIMIT_MAXCONNECTS >= 0
 	{
 		// Try setting the option. If it does not work it is ok.
 		mcode = curl_multi_setopt(CurlMultiHandle, CURLMOPT_MAXCONNECTS, (long)ELECTRA_HTTPSTREAM_CURL_PLATFORM_LIMIT_MAXCONNECTS);
