ElectraUtil: Fixed URL encode/decode with multibyte characters (UTF-8) in the path.

#rnx

diff --git a/Engine/Plugins/Media/ElectraUtil/Source/ElectraBase/Private/URLParser.cpp b/Engine/Plugins/Media/ElectraUtil/Source/ElectraBase/Private/URLParser.cpp
--- a/Engine/Plugins/Media/ElectraUtil/Source/ElectraBase/Private/URLParser.cpp
+++ b/Engine/Plugins/Media/ElectraUtil/Source/ElectraBase/Private/URLParser.cpp
@@ -158,15 +158,17 @@ namespace Electra
 
 	FString FURL_RFC3986::Get(bool bIncludeQuery, bool bIncludeFragment)
 	{
-		static const FString RequiredEscapeCharsPath(TEXT("?#"));
+		static const FString KeepCharsPath(TEXT("/"));
+		static const FString KeepCharsFragment(TEXT("!$&'()*+,;=;@/?"));
 
 		FString URL;
 		if (IsAbsolute())
 		{
+			bool bIsFile = Scheme.Equals(TEXT("file"));
 			FString Authority = GetAuthority();
 			URL = Scheme;
 			URL += TEXT(":");
-			if (Authority.Len() || Scheme.Equals(TEXT("file")))
+			if (Authority.Len() || bIsFile)
 			{
 				URL += TEXT("//");
 				URL += Authority;
@@ -177,7 +179,7 @@ namespace Electra
 				{
 					URL += TEXT("/");
 				}
-				UrlEncode(URL, Path, RequiredEscapeCharsPath);
+				UrlEncode(URL, Path, KeepCharsPath);
 			}
 			else if ((Query.Len() && bIncludeQuery) || (Fragment.Len() && bIncludeFragment))
 			{
@@ -186,7 +188,7 @@ namespace Electra
 		}
 		else
 		{
-			UrlEncode(URL, Path, RequiredEscapeCharsPath);
+			UrlEncode(URL, Path, KeepCharsPath);
 		}
 		if (Query.Len() && bIncludeQuery)
 		{
@@ -196,14 +198,15 @@ namespace Electra
 		if (Fragment.Len() && bIncludeFragment)
 		{
 			URL += TEXT("#");
-			UrlEncode(URL, Fragment, FString());
+			UrlEncode(URL, Fragment, KeepCharsFragment);
 		}
 		return URL;
 	}
 
 	FString FURL_RFC3986::GetPath(bool bIncludeQuery, bool bIncludeFragment)
 	{
-		static const FString RequiredEscapeCharsPath(TEXT("?#"));
+		static const FString KeepCharsPath(TEXT("/"));
+		static const FString KeepCharsFragment(TEXT("!$&'()*+,;=;@/?"));
 
 		FString URL;
 		if (IsAbsolute())
@@ -214,7 +217,7 @@ namespace Electra
 				{
 					URL += TEXT("/");
 				}
-				UrlEncode(URL, Path, RequiredEscapeCharsPath);
+				UrlEncode(URL, Path, KeepCharsPath);
 			}
 			else if ((Query.Len() && bIncludeQuery) || (Fragment.Len() && bIncludeFragment))
 			{
@@ -223,7 +226,7 @@ namespace Electra
 		}
 		else
 		{
-			UrlEncode(URL, Path, RequiredEscapeCharsPath);
+			UrlEncode(URL, Path, KeepCharsPath);
 		}
 		if (Query.Len() && bIncludeQuery)
 		{
@@ -233,7 +236,7 @@ namespace Electra
 		if (Fragment.Len() && bIncludeFragment)
 		{
 			URL += TEXT("#");
-			UrlEncode(URL, Fragment, FString());
+			UrlEncode(URL, Fragment, KeepCharsFragment);
 		}
 		return URL;
 	}
@@ -543,84 +546,92 @@ namespace Electra
 
 	bool FURL_RFC3986::UrlDecode(FString& OutResult, const FString& InUrlToDecode)
 	{
-		StringHelpers::FStringIterator it(InUrlToDecode);
-		while(it)
+		TArray<uint8> AsciiString;
+		AsciiString.Reserve(InUrlToDecode.Len() + 1);
+		auto IsHexChar = [](ANSICHAR In) ->bool
 		{
-			TCHAR c = *it++;
-			if (c == TCHAR('%'))
+			return (In >= ANSICHAR('0') && In <= ANSICHAR('9')) || (In >= ANSICHAR('a') && In <= ANSICHAR('f')) || (In >= ANSICHAR('A') && In <= ANSICHAR('F'));
+		};
+		auto HexVal = [](ANSICHAR In) -> uint8
+		{
+			if (In >= ANSICHAR('0') && In <= ANSICHAR('9'))
 			{
-				if (!it)
-				{
-					// '%' at the end with nothing following!
-					return false;
-				}
-				TCHAR hi = *it++;
-				if (!it)
-				{
-					// Only one char after '%'. Need two!!
-					return false;
-				}
-				TCHAR lo = *it++;
-				if (lo >= TCHAR('0') && lo <= TCHAR('9'))
-				{
-					lo -= TCHAR('0');
-				}
-				else if (lo >= TCHAR('a') && lo <= TCHAR('f'))
-				{
-					lo = lo - TCHAR('a') + 10;
-				}
-				else if (lo >= TCHAR('A') && lo <= TCHAR('F'))
-				{
-					lo = lo - TCHAR('A') + 10;
-				}
-				else
-				{
-					// Not a hex digit!
-					return false;
-				}
-				if (hi >= TCHAR('0') && hi <= TCHAR('9'))
-				{
-					hi -= TCHAR('0');
-				}
-				else if (hi >= TCHAR('a') && hi <= TCHAR('f'))
-				{
-					hi = hi - TCHAR('a') + 10;
-				}
-				else if (hi >= TCHAR('A') && hi <= TCHAR('F'))
-				{
-					hi = hi - TCHAR('A') + 10;
-				}
-				else
-				{
-					// Not a hex digit!
-					return false;
-				}
-
-				c = hi * 16 + lo;
+				return In - ANSICHAR('0');
+			}
+			else if (In >= ANSICHAR('a') && In <= ANSICHAR('f'))
+			{
+				return In - ANSICHAR('a') + 10;
+			}
+			else
+			{
+				return In - ANSICHAR('A') + 10;
+			}
+		};
+		// The URL to be decoded should strictly speaking consist only of ASCII characters, but depending on how this is used
+		// it may as well be a UTF8 encoded string. We decode UTF8 to ASCII, which is not doing anything if the URL is already
+		// only composed of ASCII characters, then we unescape it and finally convert it back to UTF8, which we need to do
+		// to get any UTF8 characters back that are represented in the URL through escaping.
+		for(const ANSICHAR* InUTF8Bytes = TCHAR_TO_UTF8(*InUrlToDecode); *InUTF8Bytes; ++InUTF8Bytes)
+		{
+			// Escaped?
+			if (*InUTF8Bytes == ANSICHAR('%') && InUTF8Bytes[1] && IsHexChar(InUTF8Bytes[1]) && InUTF8Bytes[2] && IsHexChar(InUTF8Bytes[2]))
+			{
+				uint8 hi = HexVal(*(++InUTF8Bytes));
+				uint8 lo = HexVal(*(++InUTF8Bytes));
+				AsciiString.Add((uint8)(hi * 16 + lo));
+			}
+			else
+			{
+				AsciiString.Add((uint8)*InUTF8Bytes);
 			}
-			OutResult += c;
 		}
+		AsciiString.Add((uint8)0);
+		// Convert to UTF8 to get any escaped sequences back to what they were.
+		OutResult = UTF8_TO_TCHAR((const ANSICHAR*)AsciiString.GetData());
 		return true;
 	}
 
-	bool FURL_RFC3986::UrlEncode(FString& OutResult, const FString& InUrlToEncode, const FString& InReservedChars)
+	bool FURL_RFC3986::UrlEncode(FString& OutResult, const FString& InUrlToEncode, const FString& InCharsToKeep)
 	{
-		static const FString RequiredEscapeChars(TEXT("%<>{}|\\\"^`"));
-		int32 DummyPos = 0;
-		for(StringHelpers::FStringIterator it(InUrlToEncode); it; ++it)
+		const ANSICHAR* InKeepCharsASCII = TCHAR_TO_ANSI(*InCharsToKeep);
+		auto KeepUnchanged = [InKeepCharsASCII](ANSICHAR In) -> bool
 		{
-			TCHAR c = *it;
-			if ((c >= TCHAR('a') && c <= TCHAR('z')) || (c >= TCHAR('0') && c <= TCHAR('9')) || (c >= TCHAR('A') && c <= TCHAR('Z')) || c == TCHAR('-') || c == TCHAR('_') || c == TCHAR('.') || c == TCHAR('~'))
+			for(const ANSICHAR* Res=InKeepCharsASCII; *Res; ++Res)
 			{
-				OutResult += c;
+				if (*Res == In)
+				{
+					return true;
+				}
 			}
-			else if (c <= 0x20 || c >= 0x7f || RequiredEscapeChars.FindChar(c, DummyPos) || InReservedChars.FindChar(c, DummyPos))
+			return false;
+		};
+		auto IsHexChar = [](ANSICHAR In) ->bool
+		{
+			return (In >= ANSICHAR('0') && In <= ANSICHAR('9')) || (In >= ANSICHAR('a') && In <= ANSICHAR('f')) || (In >= ANSICHAR('A') && In <= ANSICHAR('F'));
+		};
+
+		for(const ANSICHAR* InUTF8Bytes = TCHAR_TO_UTF8(*InUrlToEncode); *InUTF8Bytes; ++InUTF8Bytes)
+		{
+			uint8 c = (uint8)(*InUTF8Bytes);
+			// Unreserved character?
+			if ((c >= uint8('a') && c <= uint8('z')) || (c >= uint8('0') && c <= uint8('9')) || (c >= uint8('A') && c <= uint8('Z')) || c == uint8('-') || c == uint8('_') || c == uint8('.') || c == uint8('~'))
 			{
-				OutResult += FString::Printf(TEXT("%%%02X"), c);
+				OutResult += TCHAR(c);
+			}
+			// Already escaped?
+			else if (c == uint8('%') && InUTF8Bytes[1] && IsHexChar(InUTF8Bytes[1]) && InUTF8Bytes[2] && IsHexChar(InUTF8Bytes[2]))
+			{
+				OutResult += TCHAR(c);
+				OutResult += TCHAR(*(++InUTF8Bytes));
+				OutResult += TCHAR(*(++InUTF8Bytes));
+			}
+			else if (KeepUnchanged(*InUTF8Bytes))
+			{
+				OutResult += TCHAR(c);
 			}
 			else
 			{
-				OutResult += c;
+				OutResult += FString::Printf(TEXT("%%%02X"), c);
 			}
 		}
 		return true;
diff --git a/Engine/Plugins/Media/ElectraUtil/Source/ElectraBase/Public/Utilities/URLParser.h b/Engine/Plugins/Media/ElectraUtil/Source/ElectraBase/Public/Utilities/URLParser.h
--- a/Engine/Plugins/Media/ElectraUtil/Source/ElectraBase/Public/Utilities/URLParser.h
+++ b/Engine/Plugins/Media/ElectraUtil/Source/ElectraBase/Public/Utilities/URLParser.h
@@ -61,7 +61,9 @@ namespace Electra
 		// Decodes %XX escaped sequences into their original characters. Appends to the output. Hence in and out must not be the same.
 		static bool UrlDecode(FString& OutResult, const FString& InUrlToDecode);
 		// Encodes characters not permitted in a URL into %XX escaped sequences. Appends to the output. Hence in and out must not be the same.
-		static bool UrlEncode(FString& OutResult, const FString& InUrlToEncode, const FString& InReservedChars);
+		// A list of characters that should NOT be escaped, even if they normally would, can be passed. These characters must be strictly ASCII
+		// characters only. This is useful to prevent '/' path delimiters to be escaped if the string to be encoded represents a path.
+		static bool UrlEncode(FString& OutResult, const FString& InUrlToEncode, const FString& InCharsToKeep);
 
 		// Returns the standard port for the given scheme. An empty string is returned if none is known.
 		static FString GetStandardPortForScheme(const FString& InScheme, bool bIgnoreCase=true);
