ElectraUtil: Moving string conversion object out of loop scope to prevent it going out of scope while iterating.

#rnx

diff --git a/Engine/Plugins/Media/ElectraUtil/Source/ElectraBase/Private/URLParser.cpp b/Engine/Plugins/Media/ElectraUtil/Source/ElectraBase/Private/URLParser.cpp
--- a/Engine/Plugins/Media/ElectraUtil/Source/ElectraBase/Private/URLParser.cpp
+++ b/Engine/Plugins/Media/ElectraUtil/Source/ElectraBase/Private/URLParser.cpp
@@ -571,7 +571,8 @@ namespace Electra
 		// it may as well be a UTF8 encoded string. We decode UTF8 to ASCII, which is not doing anything if the URL is already
 		// only composed of ASCII characters, then we unescape it and finally convert it back to UTF8, which we need to do
 		// to get any UTF8 characters back that are represented in the URL through escaping.
-		for(const ANSICHAR* InUTF8Bytes = TCHAR_TO_UTF8(*InUrlToDecode); *InUTF8Bytes; ++InUTF8Bytes)
+		FTCHARToUTF8 UTF8String(*InUrlToDecode);
+		for(const ANSICHAR* InUTF8Bytes=UTF8String.Get(); *InUTF8Bytes; ++InUTF8Bytes)
 		{
 			// Escaped?
 			if (*InUTF8Bytes == ANSICHAR('%') && InUTF8Bytes[1] && IsHexChar(InUTF8Bytes[1]) && InUTF8Bytes[2] && IsHexChar(InUTF8Bytes[2]))
@@ -593,7 +594,8 @@ namespace Electra
 
 	bool FURL_RFC3986::UrlEncode(FString& OutResult, const FString& InUrlToEncode, const FString& InCharsToKeep)
 	{
-		const ANSICHAR* InKeepCharsASCII = TCHAR_TO_ANSI(*InCharsToKeep);
+		auto AnsiString = StringCast<ANSICHAR>(*InCharsToKeep);
+		const ANSICHAR* InKeepCharsASCII = AnsiString.Get();
 		auto KeepUnchanged = [InKeepCharsASCII](ANSICHAR In) -> bool
 		{
 			for(const ANSICHAR* Res=InKeepCharsASCII; *Res; ++Res)
@@ -610,7 +612,8 @@ namespace Electra
 			return (In >= ANSICHAR('0') && In <= ANSICHAR('9')) || (In >= ANSICHAR('a') && In <= ANSICHAR('f')) || (In >= ANSICHAR('A') && In <= ANSICHAR('F'));
 		};
 
-		for(const ANSICHAR* InUTF8Bytes = TCHAR_TO_UTF8(*InUrlToEncode); *InUTF8Bytes; ++InUTF8Bytes)
+		FTCHARToUTF8 UTF8String(*InUrlToEncode);
+		for(const ANSICHAR* InUTF8Bytes=UTF8String.Get(); *InUTF8Bytes; ++InUTF8Bytes)
 		{
 			uint8 c = (uint8)(*InUTF8Bytes);
 			// Unreserved character?
