Adding new PC texture sample output mode to allow for direct passing through of WebRTC otuput into Electra texture samples
#rnx

diff --git a/Engine/Plugins/Media/ElectraUtil/Source/ElectraBase/Public/MediaVideoDecoderOutputPC.h b/Engine/Plugins/Media/ElectraUtil/Source/ElectraBase/Public/MediaVideoDecoderOutputPC.h
--- a/Engine/Plugins/Media/ElectraUtil/Source/ElectraBase/Public/MediaVideoDecoderOutputPC.h
+++ b/Engine/Plugins/Media/ElectraUtil/Source/ElectraBase/Public/MediaVideoDecoderOutputPC.h
@@ -25,10 +25,11 @@ public:
 	enum class EOutputType
 	{
 		Unknown = 0,
-		SoftwareWin7,			// SW decode into buffer
+		SoftwareWin7,			// SW decode into CPU buffer
 		SoftwareWin8Plus,		// SW decode into DX11 texture
 		HardwareWin8Plus,		// HW decode into shared DX11 texture
-		HardwareDX9_DX12,		// HW decode into buffer
+		HardwareDX9_DX12,		// HW decode into CPU buffer
+		Hardware_DX,			// HW decode into DX texture on render device (DX11 or DX12)
 	};
 
 	virtual EOutputType GetOutputType() const = 0;
diff --git a/Engine/Plugins/Media/ElectraUtil/Source/ElectraSamples/Private/ElectraTextureSamplePC.cpp b/Engine/Plugins/Media/ElectraUtil/Source/ElectraSamples/Private/ElectraTextureSamplePC.cpp
--- a/Engine/Plugins/Media/ElectraUtil/Source/ElectraSamples/Private/ElectraTextureSamplePC.cpp
+++ b/Engine/Plugins/Media/ElectraUtil/Source/ElectraSamples/Private/ElectraTextureSamplePC.cpp
@@ -220,14 +220,23 @@ bool FElectraTextureSample::Convert(FTexture2DRHIRef & InDstTexture, const FConv
 	EPixelFormat Format;
 	if (VideoDecoderOutput->GetOutputType() != FVideoDecoderOutputPC::EOutputType::HardwareWin8Plus)
 	{
-		//
-		// SW decoder has decoded into a HW texture (not known to RHI) -> copy it into an RHI one
-		//
-		check(VideoDecoderOutput->GetOutputType() == FVideoDecoderOutputPC::EOutputType::SoftwareWin8Plus);
-		check(VideoDecoderOutput->GetFormat() == EPixelFormat::PF_NV12);
-
-		Format = EPixelFormat::PF_G8; // use fixed format: we flag this as NV12, too - as it is - but DX11 will only support a "higher than normal G8 texture" (with specialized access in shader)
-		bCrossDeviceCopy = false;
+		if (VideoDecoderOutput->GetOutputType() != FVideoDecoderOutputPC::EOutputType::Hardware_DX)
+		{
+			//
+			// SW decoder has decoded into a HW texture (not known to RHI) -> copy it into an RHI one
+			//
+			check(VideoDecoderOutput->GetOutputType() == FVideoDecoderOutputPC::EOutputType::SoftwareWin8Plus);
+			check(VideoDecoderOutput->GetFormat() == EPixelFormat::PF_NV12);
+
+			Format = EPixelFormat::PF_G8; // use fixed format: we flag this as NV12, too - as it is - but DX11 will only support a "higher than normal G8 texture" (with specialized access in shader)
+			bCrossDeviceCopy = false;
+		}
+		else
+		{
+			// We have a texture on the rendering device
+			Format = (VideoDecoderOutput->GetFormat() == PF_NV12) ? PF_G8 : PF_G16; // return a texture, representing the whole NV12/P010 layout as a single texture (1.5 height)
+			bCrossDeviceCopy = false;
+		}
 	}
 	else
 	{
