ElectraPlayer: Added extraction of mp4 metadata

#rnx

diff --git a/Engine/Plugins/Media/ElectraPlayer/Source/ElectraPlayerRuntime/Private/ElectraPlayer.cpp b/Engine/Plugins/Media/ElectraPlayer/Source/ElectraPlayerRuntime/Private/ElectraPlayer.cpp
--- a/Engine/Plugins/Media/ElectraPlayer/Source/ElectraPlayerRuntime/Private/ElectraPlayer.cpp
+++ b/Engine/Plugins/Media/ElectraPlayer/Source/ElectraPlayerRuntime/Private/ElectraPlayer.cpp
@@ -1868,6 +1868,12 @@ void FElectraPlayer::HandleDeferredPlayerEvents()
 				HandlePlayerEventSeekCompleted();
 				break;
 			}
+			case FPlayerMetricEventBase::EType::MediaMetadataChanged:
+			{
+				FPlayerMetricEvent_MediaMetadataChange* Ev = static_cast<FPlayerMetricEvent_MediaMetadataChange*>(Event.Get());
+				HandlePlayerMediaMetadataChanged(Ev->NewMetadata);
+				break;
+			}
 			case FPlayerMetricEventBase::EType::Error:
 			{
 				FPlayerMetricEvent_Error* Ev = static_cast<FPlayerMetricEvent_Error*>(Event.Get());
@@ -2586,6 +2592,17 @@ void FElectraPlayer::HandlePlayerEventSeekCompleted()
 	MediaStateOnSeekFinished();
 }
 
+void FElectraPlayer::HandlePlayerMediaMetadataChanged(const FString& InMetadata)
+{
+	TSharedPtr<IElectraPlayerAdapterDelegate, ESPMode::ThreadSafe> PinnedAdapterDelegate = AdapterDelegate.Pin();
+	if (PinnedAdapterDelegate.IsValid())
+	{
+		// Send out the metadata through the option query interface since we have no other means
+		// to deliver it through Media Framework at the moment.
+		/*FVariantValue Result =*/ PinnedAdapterDelegate->QueryOptions(IElectraPlayerAdapterDelegate::EOptionType::MediaMetadataUpdate, FVariantValue(InMetadata));
+	}
+}
+
 void FElectraPlayer::HandlePlayerEventError(const FString& ErrorReason)
 {
 	bHasPendingError = true;
diff --git a/Engine/Plugins/Media/ElectraPlayer/Source/ElectraPlayerRuntime/Private/ElectraPlayer.h b/Engine/Plugins/Media/ElectraPlayer/Source/ElectraPlayerRuntime/Private/ElectraPlayer.h
--- a/Engine/Plugins/Media/ElectraPlayer/Source/ElectraPlayerRuntime/Private/ElectraPlayer.h
+++ b/Engine/Plugins/Media/ElectraPlayer/Source/ElectraPlayerRuntime/Private/ElectraPlayer.h
@@ -202,6 +202,7 @@ private:
 			JumpInPlayPosition,
 			PlaybackStopped,
 			SeekCompleted,
+			MediaMetadataChanged,
 			Error,
 			LogMessage,
 			DroppedVideoFrame,
@@ -282,6 +283,11 @@ private:
 		FTimeValue FromTime;
 		Metrics::ETimeJumpReason TimejumpReason;
 	};
+	struct FPlayerMetricEvent_MediaMetadataChange : public FPlayerMetricEventBase
+	{
+		FPlayerMetricEvent_MediaMetadataChange(const TSharedPtrTS<Electra::UtilsMP4::FMetadataParser>& InMetadata) : FPlayerMetricEventBase(EType::MediaMetadataChanged), NewMetadata(InMetadata->GetAsJSON()) {}
+		FString NewMetadata;
+	};
 	struct FPlayerMetricEvent_Error : public FPlayerMetricEventBase
 	{
 		FPlayerMetricEvent_Error(const FString& InErrorReason) : FPlayerMetricEventBase(EType::Error), ErrorReason(InErrorReason) {}
@@ -358,6 +364,8 @@ private:
 	{ DeferredPlayerEvents.Enqueue(MakeSharedTS<FPlayerMetricEventBase>(FPlayerMetricEventBase::EType::PlaybackStopped)); }
 	virtual void ReportSeekCompleted() override
 	{ DeferredPlayerEvents.Enqueue(MakeSharedTS<FPlayerMetricEventBase>(FPlayerMetricEventBase::EType::SeekCompleted)); }
+	virtual void ReportMediaMetadataChanged(TSharedPtrTS<Electra::UtilsMP4::FMetadataParser> Metadata) override
+	{ DeferredPlayerEvents.Enqueue(MakeSharedTS<FPlayerMetricEvent_MediaMetadataChange>(Metadata)); }
 	virtual void ReportError(const FString& ErrorReason) override
 	{ DeferredPlayerEvents.Enqueue(MakeSharedTS<FPlayerMetricEvent_Error>(ErrorReason)); }
 	virtual void ReportLogMessage(IInfoLog::ELevel InLogLevel, const FString& InLogMessage, int64 InPlayerWallclockMilliseconds) override
@@ -795,6 +803,7 @@ private:
 	void HandlePlayerEventJumpInPlayPosition(const FTimeValue& ToNewTime, const FTimeValue& FromTime, Metrics::ETimeJumpReason TimejumpReason);
 	void HandlePlayerEventPlaybackStopped();
 	void HandlePlayerEventSeekCompleted();
+	void HandlePlayerMediaMetadataChanged(const FString& InMetadata);
 	void HandlePlayerEventError(const FString& ErrorReason);
 	void HandlePlayerEventLogMessage(IInfoLog::ELevel InLogLevel, const FString& InLogMessage, int64 InPlayerWallclockMilliseconds);
 	void HandlePlayerEventDroppedVideoFrame();
diff --git a/Engine/Plugins/Media/ElectraPlayer/Source/ElectraPlayerRuntime/Private/Runtime/Demuxer/ParserISO14496-12.cpp b/Engine/Plugins/Media/ElectraPlayer/Source/ElectraPlayerRuntime/Private/Runtime/Demuxer/ParserISO14496-12.cpp
--- a/Engine/Plugins/Media/ElectraPlayer/Source/ElectraPlayerRuntime/Private/Runtime/Demuxer/ParserISO14496-12.cpp
+++ b/Engine/Plugins/Media/ElectraPlayer/Source/ElectraPlayerRuntime/Private/Runtime/Demuxer/ParserISO14496-12.cpp
@@ -309,7 +309,6 @@ private:
 			return BoxReader;
 		}
 
-
 		FMP4Box* GetCurrentFtypBox() const
 		{
 			return CurrentFtypBox;
@@ -359,7 +358,7 @@ private:
 			return CurrentMediaHandlerBox;
 		}
 
-		UEMediaError ReadAndParseNextBox(FMP4Box* ParentBox);
+		UEMediaError ReadAndParseNextBox(FMP4Box* ParentBox, bool bWarnIfUnsupported=true);
 
 		const FMP4Box* GetRootBox() const
 		{
@@ -1505,6 +1504,11 @@ private:
 			return HandlerType;
 		}
 
+		uint32 GetReservedValue(int32 Index) const
+		{
+			return Index >=0 && Index <= 2 ? Reserved[Index] : 0;
+		}
+
 		FString GetHandlerName() const
 		{
 			return NameUTF8;
@@ -1523,9 +1527,9 @@ private:
 			RETURN_IF_ERROR(FMP4BoxFull::ReadAndParseAttributes(ParseInfo));
 			RETURN_IF_ERROR(ParseInfo->Reader()->Read(Value32));				// pre_defined (in QuickTime this was "Component Type")
 			RETURN_IF_ERROR(ParseInfo->Reader()->Read(HandlerType));			// handler_type
-			RETURN_IF_ERROR(ParseInfo->Reader()->Read(Value32));				// reserved[0] (In QuickTime this was "Component Manufacturer")
-			RETURN_IF_ERROR(ParseInfo->Reader()->Read(Value32));				// reserved[1] (In QuickTime this was "Component Flags")
-			RETURN_IF_ERROR(ParseInfo->Reader()->Read(Value32));				// reserved[2] (In QuickTime this was "Component Flags Mask")
+			RETURN_IF_ERROR(ParseInfo->Reader()->Read(Reserved[0]));			// reserved[0] (In QuickTime this was "Component Manufacturer")
+			RETURN_IF_ERROR(ParseInfo->Reader()->Read(Reserved[1]));			// reserved[1] (In QuickTime this was "Component Flags")
+			RETURN_IF_ERROR(ParseInfo->Reader()->Read(Reserved[2]));			// reserved[2] (In QuickTime this was "Component Flags Mask")
 
 			uint32 BytesRemaining = BoxSize - (ParseInfo->Reader()->GetCurrentReadOffset() - StartOffset);
 			RETURN_IF_ERROR(ParseInfo->Reader()->ReadString(NameUTF8, BytesRemaining));
@@ -1539,6 +1543,7 @@ private:
 	private:
 		FString			NameUTF8;
 		uint32			HandlerType;
+		uint32			Reserved[3];
 	};
 
 
@@ -4894,30 +4899,11 @@ private:
 		{
 			UEMediaError Error = UEMEDIA_ERROR_OK;
 
+			bIsLeafBox = false;
 			int32 BytesRemaining = (int32) (BoxSize - (ParseInfo->Reader()->GetCurrentReadOffset() - StartOffset));
-			if (BytesRemaining > 8)
+			if (BytesRemaining >= 8)
 			{
-				// Now that there are additional boxes we ourselves are no longer a leaf box.
-				bIsLeafBox = false;
-				// This will read all the boxes in here and add them as children.
-				for(; BytesRemaining>8; BytesRemaining = (int32) (BoxSize - (ParseInfo->Reader()->GetCurrentReadOffset() - StartOffset)))
-				{
-					IParserISO14496_12::FBoxType	ChildBoxType;
-					int64							ChildBoxSize;
-					uint8							ChildBoxUUID[16];
-
-					int64 BoxStartOffset = ParseInfo->Reader()->GetCurrentReadOffset();
-					RETURN_IF_ERROR(ParseInfo->ReadBoxTypeAndSize(ChildBoxType, ChildBoxSize, ChildBoxUUID));
-#if MEDIA_DEBUG_HAS_BOX_NAMES
-					char boxName[4];
-					*((uint32*)&boxName) = MEDIA_TO_BIG_ENDIAN(ChildBoxType);
-#endif
-					int64 BoxDataOffset = ParseInfo->Reader()->GetCurrentReadOffset();
-					FMP4Box* NextBox = new FMP4BoxIgnored(ChildBoxType, ChildBoxSize, BoxStartOffset, BoxDataOffset, true);
-
-					AddChildBox(NextBox);
-					RETURN_IF_ERROR(ParseInfo->ReadAndParseNextBox(NextBox));
-				}
+				RETURN_IF_ERROR(ParseInfo->ReadAndParseNextBox(this, false));
 			}
 			// If there is still something left that is not a complete box, skip over it.
 			// There may be a 32 bit integer zero value to terminate the list.
@@ -4932,6 +4918,72 @@ private:
 	};
 
 
+	/**
+	 * 'meta' box. ISO/IEC 14496-12:2015 - 8.11.1 Meta Box
+	 */
+	class FMP4BoxMETA : public FMP4BoxFull,	public IParserISO14496_12::IMetadata
+	{
+	public:
+		FMP4BoxMETA(IParserISO14496_12::FBoxType InBoxType, int64 InBoxSize, int64 InStartOffset, int64 InDataOffset, bool bInIsLeafBox)
+			: FMP4BoxFull(InBoxType, InBoxSize, InStartOffset, InDataOffset, bInIsLeafBox)
+		{
+		}
+
+		virtual ~FMP4BoxMETA() = default;
+
+	private:
+		FMP4BoxMETA() = delete;
+		FMP4BoxMETA(const FMP4BoxMETA&) = delete;
+
+
+		uint32 GetHandler() const override
+		{
+			return Handler ? Handler->GetHandlerType() : 0;
+		}
+		uint32 GetReserved(int32 InIndex0to2) const override
+		{
+			return Handler ? Handler->GetReservedValue(InIndex0to2) : 0;
+		}
+		int32 GetNumChildBoxes() const override
+		{
+			return GetNumberOfChildren();
+		}
+		uint32 GetChildBoxType(int32 InChildIndex) const override
+		{
+			const FMP4Box* Child = GetChildBox(InChildIndex);
+			return Child ? (uint32)Child->GetType() : 0;
+		}
+		uint32 GetChildBoxDataSize(int32 InChildIndex) const override
+		{
+			const FMP4Box* Child = GetChildBox(InChildIndex);
+			return Child ? (uint32)Child->GetDataSize() : 0;
+		}
+		const void* GetChildBoxData(int32 InChildIndex) const override
+		{
+			const FMP4Box* Child = GetChildBox(InChildIndex);
+			return Child ? Child->GetFillerData() ? Child->GetFillerData()->Data : nullptr : nullptr;
+		}
+
+	protected:
+		UEMediaError ReadAndParseAttributes(FMP4ParseInfo* ParseInfo) override
+		{
+			UEMediaError Error = UEMEDIA_ERROR_OK;
+			RETURN_IF_ERROR(FMP4BoxFull::ReadAndParseAttributes(ParseInfo));
+
+			bIsLeafBox = false;
+			int32 BytesRemaining = (int32) (BoxSize - (ParseInfo->Reader()->GetCurrentReadOffset() - StartOffset));
+			if (BytesRemaining >= 8)
+			{
+				RETURN_IF_ERROR(ParseInfo->ReadAndParseNextBox(this, false));
+			}
+			Handler = static_cast<const FMP4BoxHDLR*>(FindBox(kBox_hdlr, 0));
+			return Error;
+		}
+	private:
+		const FMP4BoxHDLR* Handler = nullptr;
+	};
+
+
 	FMP4ParseInfo::~FMP4ParseInfo()
 	{
 		delete RootBox;
@@ -5000,7 +5052,7 @@ private:
 	}
 
 
-	UEMediaError FMP4ParseInfo::ReadAndParseNextBox(FMP4Box* ParentBox)
+	UEMediaError FMP4ParseInfo::ReadAndParseNextBox(FMP4Box* ParentBox, bool bWarnIfUnsupported)
 	{
 		IParserISO14496_12::FBoxType	BoxType;
 		int64							BoxSize;
@@ -5249,6 +5301,9 @@ private:
 						case FMP4Box::kBox_udta:
 							NextBox = new FMP4BoxUDTA(BoxType, BoxSize, BoxStartOffset, BoxDataOffset, true);
 							break;
+						case FMP4Box::kBox_meta:
+							NextBox = new FMP4BoxMETA(BoxType, BoxSize, BoxStartOffset, BoxDataOffset, true);
+							break;
 						case FMP4Box::kBox_skip:
 						case FMP4Box::kBox_VOID:
 							NextBox = new FMP4BoxIgnored(BoxType, BoxSize, BoxStartOffset, BoxDataOffset, true);
@@ -5306,7 +5361,10 @@ private:
 
 						default:
 #if MEDIA_DEBUG_HAS_BOX_NAMES
-							UE_LOG(LogElectraMP4Parser, Log, TEXT("Ignoring mp4 box: '%c%c%c%c' %lld B @ %lld"), boxName[0]?boxName[0]:'?', boxName[1]?boxName[1]:'?', boxName[2]?boxName[2]:'?', boxName[3]?boxName[3]:'?', (long long int)BoxSize, (long long int)BoxStartOffset);
+							if (bWarnIfUnsupported)
+							{
+								UE_LOG(LogElectraMP4Parser, Log, TEXT("Ignoring mp4 box: '%c%c%c%c' %lld B @ %lld"), boxName[0]?boxName[0]:'?', boxName[1]?boxName[1]:'?', boxName[2]?boxName[2]:'?', boxName[3]?boxName[3]:'?', (long long int)BoxSize, (long long int)BoxStartOffset);
+							}
 #endif
 							NextBox = new FMP4BoxIgnored(BoxType, BoxSize, BoxStartOffset, BoxDataOffset, true);
 							break;
@@ -5316,7 +5374,7 @@ private:
 						++NumTotalBoxesParsed;
 						ParentBox->AddChildBox(NextBox);
 						IncreaseNestingLevel();
-						Error = ReadAndParseNextBox(NextBox);
+						Error = ReadAndParseNextBox(NextBox, bWarnIfUnsupported);
 						DecreaseNestingLevel();
 						if (Error == UEMEDIA_ERROR_OK)
 						{
@@ -5441,6 +5499,7 @@ private:
 		const ITrack* GetTrackByTrackID(int32 TrackID) const override;
 		const ISegmentIndex* GetSegmentIndexByIndex(int32 Index) const override;
 		const IEventMessage* GetEventMessageByIndex(int32 Index) const override;
+		const IMetadata* GetMetadata(EBaseBoxType InFromBox) const override;
 
 		TSharedPtrTS<IAllTrackIterator> CreateAllTrackIteratorByFilePos(int64 InFromFilePos) const override;
 
@@ -6843,6 +6902,7 @@ private:
 								Track->BitrateInfo.BufferSizeDB = Track->CodecSpecificDataMP4A.GetBufferSize();
 								Track->BitrateInfo.MaxBitrate = Track->CodecSpecificDataMP4A.GetMaxBitrate();
 								Track->BitrateInfo.AvgBitrate = Track->CodecSpecificDataMP4A.GetAvgBitrate();
+								Track->CodecInformation.SetBitrate(Track->BitrateInfo.MaxBitrate);
 
 								bGotAudioFormat = true;
 
@@ -7512,6 +7572,23 @@ private:
 		return ParsedData ? static_cast<FMP4BoxEMSG*>(ParsedData->GetEMSGBox(Index)) : nullptr;
 	}
 
+	const IParserISO14496_12::IMetadata* FParserISO14496_12::GetMetadata(IParserISO14496_12::EBaseBoxType InFromBox) const
+	{
+		// We need parsed data and a 'moov' box. We do not look at top-level 'meta' boxes here.
+		if (!ParsedData || (InFromBox == EBaseBoxType::Moov && !ParsedData->GetCurrentMoovBox()) || (InFromBox == EBaseBoxType::Moof && !ParsedData->GetCurrentMoofBox()))
+		{
+			return nullptr;
+		}
+		// First try 'udta'->'meta', which is the most common location.
+		const FMP4Box* Box = InFromBox == EBaseBoxType::Moov ? ParsedData->GetCurrentMoovBox() : ParsedData->GetCurrentMoofBox();
+		const FMP4BoxMETA* Meta = Box ? static_cast<const FMP4BoxMETA*>(Box->GetBoxPath(FMP4Box::kBox_udta, FMP4Box::kBox_meta)) : nullptr;
+		if (!Meta)
+		{
+			Meta = Box ? static_cast<const FMP4BoxMETA*>(Box->GetBoxPath(FMP4Box::kBox_meta)) : nullptr;
+		}
+		return Meta;
+	}
+
 	int32 FParserISO14496_12::GetNumberOfBrands() const
 	{
 		const FMP4BoxFTYP* Box = static_cast<const FMP4BoxFTYP*>(ParsedData->GetCurrentStypBox() ? ParsedData->GetCurrentStypBox() : ParsedData->GetCurrentFtypBox());
diff --git a/Engine/Plugins/Media/ElectraPlayer/Source/ElectraPlayerRuntime/Private/Runtime/Demuxer/ParserISO14496-12.h b/Engine/Plugins/Media/ElectraPlayer/Source/ElectraPlayerRuntime/Private/Runtime/Demuxer/ParserISO14496-12.h
--- a/Engine/Plugins/Media/ElectraPlayer/Source/ElectraPlayerRuntime/Private/Runtime/Demuxer/ParserISO14496-12.h
+++ b/Engine/Plugins/Media/ElectraPlayer/Source/ElectraPlayerRuntime/Private/Runtime/Demuxer/ParserISO14496-12.h
@@ -165,6 +165,17 @@ namespace Electra
 
 		virtual int32 GetNumberOfEventMessages() const = 0;
 
+		class IMetadata
+		{
+		public:
+			virtual uint32 GetHandler() const = 0;
+			virtual uint32 GetReserved(int32 InIndex0to2) const = 0;
+			virtual int32 GetNumChildBoxes() const = 0;
+			virtual uint32 GetChildBoxType(int32 InChildIndex) const = 0;
+			virtual uint32 GetChildBoxDataSize(int32 InChildIndex) const = 0;
+			virtual const void* GetChildBoxData(int32 InChildIndex) const = 0;
+		};
+
 		class ITrackIterator
 		{
 		public:
@@ -299,6 +310,13 @@ namespace Electra
 			virtual const TArray<uint8>& GetMessageData() const = 0;
 		};
 
+		enum class EBaseBoxType
+		{
+			Moov,
+			Moof
+		};
+
+		virtual const IMetadata* GetMetadata(EBaseBoxType InFromBox) const = 0;
 
 		virtual TSharedPtrTS<IAllTrackIterator> CreateAllTrackIteratorByFilePos(int64 InFromFilePos) const = 0;
 
diff --git a/Engine/Plugins/Media/ElectraPlayer/Source/ElectraPlayerRuntime/Private/Runtime/Player/AdaptiveStreamingPlayer.cpp b/Engine/Plugins/Media/ElectraPlayer/Source/ElectraPlayerRuntime/Private/Runtime/Player/AdaptiveStreamingPlayer.cpp
--- a/Engine/Plugins/Media/ElectraPlayer/Source/ElectraPlayerRuntime/Private/Runtime/Player/AdaptiveStreamingPlayer.cpp
+++ b/Engine/Plugins/Media/ElectraPlayer/Source/ElectraPlayerRuntime/Private/Runtime/Player/AdaptiveStreamingPlayer.cpp
@@ -494,6 +494,9 @@ void FAdaptiveStreamingPlayer::FireSyncEvent(TSharedPtrTS<FMetricEvent> Event)
 			case FMetricEvent::EType::SeekCompleted:
 				Listeners[i]->ReportSeekCompleted();
 				break;
+			case FMetricEvent::EType::MediaMetadataChanged:
+				Listeners[i]->ReportMediaMetadataChanged(Event->Param.MediaMetadataChange.NewMetadata);
+				break;
 			case FMetricEvent::EType::Errored:
 				Listeners[i]->ReportError(Event->Param.ErrorDetail.GetPrintable());
 				break;
@@ -1797,6 +1800,12 @@ void FAdaptiveStreamingPlayer::HandleMetadataChanges()
 			}
 		}
 	}
+
+	// Handle media metadata updates.
+	if (MediaMetadataUpdates.Handle(CurrentTime))
+	{
+		DispatchEvent(FMetricEvent::ReportMediaMetadataChanged(MediaMetadataUpdates.GetActive()));
+	}
 }
 
 
@@ -1871,6 +1880,12 @@ void FAdaptiveStreamingPlayer::HandleSessionMessage(TSharedPtrTS<IPlayerMessage>
 			}
 		}
 	}
+	// Media metadata changed?
+	else if (SessionMessage->GetType() == FPlaylistMetadataUpdateMessage::Type())
+	{
+		FPlaylistMetadataUpdateMessage* pMsg = static_cast<FPlaylistMetadataUpdateMessage*>(SessionMessage.Get());
+		MediaMetadataUpdates.AddEntry(pMsg->GetValidFrom(), pMsg->GetMetadata());
+	}
 	// License key?
 	else if (SessionMessage->GetType() == FLicenseKeyMessage::Type())
 	{
@@ -4295,6 +4310,7 @@ void FAdaptiveStreamingPlayer::InternalStop(bool bHoldCurrentFrame)
 	ActivePeriodCriticalSection.Unlock();
 	UpcomingPeriods.Empty();
 	MetadataHandlingState.Reset();
+	MediaMetadataUpdates.Reset();
 	NextLoopStates.Empty();
 	SeekVars.Reset();
 
diff --git a/Engine/Plugins/Media/ElectraPlayer/Source/ElectraPlayerRuntime/Private/Runtime/Player/AdaptiveStreamingPlayerABR.cpp b/Engine/Plugins/Media/ElectraPlayer/Source/ElectraPlayerRuntime/Private/Runtime/Player/AdaptiveStreamingPlayerABR.cpp
--- a/Engine/Plugins/Media/ElectraPlayer/Source/ElectraPlayerRuntime/Private/Runtime/Player/AdaptiveStreamingPlayerABR.cpp
+++ b/Engine/Plugins/Media/ElectraPlayer/Source/ElectraPlayerRuntime/Private/Runtime/Player/AdaptiveStreamingPlayerABR.cpp
@@ -70,6 +70,7 @@ namespace Electra
 		void ReportJumpInPlayPosition(const FTimeValue& ToNewTime, const FTimeValue& FromTime, Metrics::ETimeJumpReason TimejumpReason) override {}
 		void ReportPlaybackStopped() override {}
 		void ReportSeekCompleted() override {}
+		void ReportMediaMetadataChanged(TSharedPtrTS<UtilsMP4::FMetadataParser> Metadata) override {}
 		void ReportError(const FString& ErrorReason) override {}
 		void ReportLogMessage(IInfoLog::ELevel InLogLevel, const FString& LogMessage, int64 PlayerWallclockMilliseconds) override {}
 		void ReportDroppedVideoFrame() override {}
diff --git a/Engine/Plugins/Media/ElectraPlayer/Source/ElectraPlayerRuntime/Private/Runtime/Player/AdaptiveStreamingPlayerInternal.h b/Engine/Plugins/Media/ElectraPlayer/Source/ElectraPlayerRuntime/Private/Runtime/Player/AdaptiveStreamingPlayerInternal.h
--- a/Engine/Plugins/Media/ElectraPlayer/Source/ElectraPlayerRuntime/Private/Runtime/Player/AdaptiveStreamingPlayerInternal.h
+++ b/Engine/Plugins/Media/ElectraPlayer/Source/ElectraPlayerRuntime/Private/Runtime/Player/AdaptiveStreamingPlayerInternal.h
@@ -21,6 +21,8 @@
 #include "Player/PlayerEntityCache.h"
 #include "Player/AdaptiveStreamingPlayerABR.h"
 
+#include "Utilities/UtilsMP4.h"
+
 #include "ElectraCDM.h"
 
 #include "InfoLog.h"
@@ -564,6 +566,7 @@ struct FMetricEvent
 		PlaybackJumped,
 		PlaybackStopped,
 		SeekCompleted,
+		MediaMetadataChanged,
 		LicenseKey,
 		Errored,
 		LogMessage,
@@ -600,6 +603,10 @@ struct FMetricEvent
 		{
 			bool bWasAlreadyThere;
 		};
+		struct FMediaMedataDataChanged
+		{
+			TSharedPtrTS<UtilsMP4::FMetadataParser> NewMetadata;
+		};
 		FString								URL;
 		Metrics::EBufferingReason			BufferingReason;
 		Metrics::FBufferStats				BufferStats;
@@ -612,6 +619,7 @@ struct FMetricEvent
 		FStreamCodecInformation				CodecFormatChange;
 		FTimeJumped							TimeJump;
 		FSeekComplete						SeekComplete;
+		FMediaMedataDataChanged				MediaMetadataChange;
 		FErrorDetail						ErrorDetail;
 		FLogMessage							LogMessage;
 	};
@@ -785,6 +793,13 @@ struct FMetricEvent
 		Evt->Param.SeekComplete.bWasAlreadyThere = bWasAlreadyThere;
 		return Evt;
 	}
+	static TSharedPtrTS<FMetricEvent> ReportMediaMetadataChanged(const TSharedPtrTS<UtilsMP4::FMetadataParser>& InNextMetadata)
+	{
+		TSharedPtrTS<FMetricEvent> Evt = MakeSharedTS<FMetricEvent>();
+		Evt->Type = EType::MediaMetadataChanged;
+		Evt->Param.MediaMetadataChange.NewMetadata = InNextMetadata;
+		return Evt;
+	}
 	static TSharedPtrTS<FMetricEvent> ReportError(const FErrorDetail& ErrorDetail)
 	{
 		TSharedPtrTS<FMetricEvent> Evt = MakeSharedTS<FMetricEvent>();
@@ -1873,6 +1888,42 @@ private:
 	};
 
 
+	struct FMediaMetadataUpdate
+	{
+		FMediaMetadataUpdate()
+		{
+			Reset();
+		}
+		void Reset()
+		{
+			NextEntries.Empty();
+			// Do not reset the active metadata.
+		}
+		void AddEntry(const FTimeValue& InValidFrom, const TSharedPtrTS<UtilsMP4::FMetadataParser>& InMetadata)
+		{
+			FEntry& e = NextEntries.Emplace_GetRef();
+			e.ValidFrom = InValidFrom.IsValid() ? InValidFrom : FTimeValue::GetZero();
+			e.Metadata = InMetadata;
+			NextEntries.StableSort([](const FEntry& a, const FEntry& b)
+			{
+				return a.ValidFrom < b.ValidFrom;
+			});
+		}
+		bool Handle(const FTimeValue& InAtTime);
+		TSharedPtrTS<UtilsMP4::FMetadataParser> GetActive() const
+		{ 
+			return ActiveMetadata; 
+		}
+
+		struct FEntry
+		{
+			FTimeValue ValidFrom;
+			TSharedPtrTS<UtilsMP4::FMetadataParser> Metadata;
+		};
+		TArray<FEntry> NextEntries;
+		TSharedPtrTS<UtilsMP4::FMetadataParser> ActiveMetadata;
+	};
+
 	struct FMetadataHandlingState
 	{
 		FMetadataHandlingState()
@@ -2129,6 +2180,7 @@ private:
 	TArray<FPeriodInformation>											ActivePeriods;
 	TArray<FPeriodInformation>											UpcomingPeriods;
 	FMetadataHandlingState												MetadataHandlingState;
+	FMediaMetadataUpdate												MediaMetadataUpdates;
 	TSharedPtrTS<IManifest::IPlayPeriod>								InitialPlayPeriod;
 	TSharedPtrTS<IManifest::IPlayPeriod>								CurrentPlayPeriodVideo;
 	TSharedPtrTS<IManifest::IPlayPeriod>								CurrentPlayPeriodAudio;
diff --git a/Engine/Plugins/Media/ElectraPlayer/Source/ElectraPlayerRuntime/Private/Runtime/Player/AdaptiveStreamingPlayerManifest.cpp b/Engine/Plugins/Media/ElectraPlayer/Source/ElectraPlayerRuntime/Private/Runtime/Player/AdaptiveStreamingPlayerManifest.cpp
--- a/Engine/Plugins/Media/ElectraPlayer/Source/ElectraPlayerRuntime/Private/Runtime/Player/AdaptiveStreamingPlayerManifest.cpp
+++ b/Engine/Plugins/Media/ElectraPlayer/Source/ElectraPlayerRuntime/Private/Runtime/Player/AdaptiveStreamingPlayerManifest.cpp
@@ -343,6 +343,37 @@ void FAdaptiveStreamingPlayer::UpdateManifest()
 	}
 }
 
+
+bool FAdaptiveStreamingPlayer::FMediaMetadataUpdate::Handle(const FTimeValue& InAtTime)
+{
+	TSharedPtrTS<UtilsMP4::FMetadataParser> NextMetadata;
+	while(NextEntries.Num())
+	{
+		if (!InAtTime.IsValid() && !ActiveMetadata.IsValid())
+		{
+			NextMetadata = NextEntries[0].Metadata;
+			break;
+		}
+		else if (NextEntries[0].ValidFrom <= InAtTime)
+		{
+			NextMetadata = NextEntries[0].Metadata;
+			NextEntries.RemoveAt(0);
+		}
+		else
+		{
+			break;
+		}
+	}
+	if (NextMetadata.IsValid())
+	{
+		bool bChanged = !ActiveMetadata.IsValid() || ActiveMetadata->IsDifferentFrom(*NextMetadata);
+		ActiveMetadata = MoveTemp(NextMetadata);
+		return bChanged;
+	}
+	return false;
+}
+
+
 } // namespace Electra
 
 
diff --git a/Engine/Plugins/Media/ElectraPlayer/Source/ElectraPlayerRuntime/Private/Runtime/Player/AdaptiveStreamingPlayerMetrics.h b/Engine/Plugins/Media/ElectraPlayer/Source/ElectraPlayerRuntime/Private/Runtime/Player/AdaptiveStreamingPlayerMetrics.h
--- a/Engine/Plugins/Media/ElectraPlayer/Source/ElectraPlayerRuntime/Private/Runtime/Player/AdaptiveStreamingPlayerMetrics.h
+++ b/Engine/Plugins/Media/ElectraPlayer/Source/ElectraPlayerRuntime/Private/Runtime/Player/AdaptiveStreamingPlayerMetrics.h
@@ -8,6 +8,7 @@
 #include "Player/AdaptiveStreamingPlayerABR_State.h"
 #include "Player/Playlist.h"
 #include "ElectraHTTPStream.h"
+#include "Utilities/UtilsMP4.h"
 
 namespace Electra
 {
@@ -330,6 +331,11 @@ public:
 	 */
 	virtual void ReportSeekCompleted() = 0;
 
+	/**
+	 * Called when the media metadata has changed.
+	 */
+	virtual void ReportMediaMetadataChanged(TSharedPtrTS<UtilsMP4::FMetadataParser> Metadata) = 0;
+
 	/**
 	 * Called when an error occurs. Errors always result in termination of playback.
 	 */
diff --git a/Engine/Plugins/Media/ElectraPlayer/Source/ElectraPlayerRuntime/Private/Runtime/Player/Manifest.h b/Engine/Plugins/Media/ElectraPlayer/Source/ElectraPlayerRuntime/Private/Runtime/Player/Manifest.h
--- a/Engine/Plugins/Media/ElectraPlayer/Source/ElectraPlayerRuntime/Private/Runtime/Player/Manifest.h
+++ b/Engine/Plugins/Media/ElectraPlayer/Source/ElectraPlayerRuntime/Private/Runtime/Player/Manifest.h
@@ -1,6 +1,6 @@
 // Copyright Epic Games, Inc. All Rights Reserved.
-
 #pragma once
+
 #include "PlayerCore.h"
 #include "PlayerTime.h"
 #include "StreamTypes.h"
@@ -10,7 +10,7 @@
 #include "Player/AdaptiveStreamingPlayerMetrics.h"
 #include "Player/PlayerSessionServices.h"
 #include "ParameterDictionary.h"
-
+#include "Utilities/UtilsMP4.h"
 
 namespace Electra
 {
@@ -511,6 +511,48 @@ namespace Electra
 	};
 
 
+	/**
+	 * Playlist metadata changed.
+	 */
+	class FPlaylistMetadataUpdateMessage : public IPlayerMessage
+	{
+	public:
+		static TSharedPtrTS<IPlayerMessage> Create(FTimeValue InValidFrom, TSharedPtrTS<UtilsMP4::FMetadataParser> InMetadata)
+		{
+			TSharedPtrTS<FPlaylistMetadataUpdateMessage> p(new FPlaylistMetadataUpdateMessage(InValidFrom, InMetadata));
+			return p;
+		}
+
+		static const FString& Type()
+		{
+			static FString TypeName("PlaylistMetadataUpdate");
+			return TypeName;
+		}
+
+		virtual const FString& GetType() const
+		{
+			return Type();
+		}
+
+		FTimeValue GetValidFrom() const
+		{
+			return ValidFrom;
+		}
+		
+		TSharedPtrTS<UtilsMP4::FMetadataParser> GetMetadata() const
+		{
+			return Metadata;
+		}
+
+	private:
+		FPlaylistMetadataUpdateMessage(FTimeValue InValidFrom, TSharedPtrTS<UtilsMP4::FMetadataParser> InMetadata)
+			: ValidFrom(InValidFrom)
+			, Metadata(MoveTemp(InMetadata))
+		{ }
+		FTimeValue ValidFrom;
+		TSharedPtrTS<UtilsMP4::FMetadataParser> Metadata;
+	};
+
 
 } // namespace Electra
 
diff --git a/Engine/Plugins/Media/ElectraPlayer/Source/ElectraPlayerRuntime/Private/Runtime/Player/mp4/ManifestMP4.cpp b/Engine/Plugins/Media/ElectraPlayer/Source/ElectraPlayerRuntime/Private/Runtime/Player/mp4/ManifestMP4.cpp
--- a/Engine/Plugins/Media/ElectraPlayer/Source/ElectraPlayerRuntime/Private/Runtime/Player/mp4/ManifestMP4.cpp
+++ b/Engine/Plugins/Media/ElectraPlayer/Source/ElectraPlayerRuntime/Private/Runtime/Player/mp4/ManifestMP4.cpp
@@ -10,6 +10,7 @@
 #include "Utilities/HashFunctions.h"
 #include "Utilities/TimeUtilities.h"
 #include "HAL/LowLevelMemTracker.h"
+#include "Utilities/UtilsMP4.h"
 #include "ElectraPlayerPrivate.h"
 #include "Player/AdaptiveStreamingPlayerResourceRequest.h"
 #include "Player/AdaptivePlayerOptionKeynames.h"
@@ -642,6 +643,7 @@ FErrorDetail FManifestMP4Internal::FTimelineAssetMP4::Build(IPlayerSessionServic
 {
 	PlayerSessionServices = InPlayerSessionServices;
 	MediaURL = URL;
+
 	// Go over the supported tracks and create an internal manifest-like structure for the player to work with.
 	LongestTrackDuration = FTimeValue::GetZero();
 	for(int32 nTrack=0,nMaxTrack=MP4Parser->GetNumberOfTracks(); nTrack < nMaxTrack; ++nTrack)
@@ -701,6 +703,27 @@ FErrorDetail FManifestMP4Internal::FTimelineAssetMP4::Build(IPlayerSessionServic
 	// Hold on to the parsed MOOV box for future reference.
 	MoovBoxParser = MP4Parser;
 
+	// Get the metadata, if any.
+	const IParserISO14496_12::IMetadata* md = MP4Parser->GetMetadata(IParserISO14496_12::EBaseBoxType::Moov);
+	if (md)
+	{
+		uint32 hdlr = md->GetHandler();
+		uint32 res0 = md->GetReserved(0);
+		TArray<UtilsMP4::FMetadataParser::FBoxInfo> Boxes;
+		for(int32 i=0, iMax=md->GetNumChildBoxes(); i<iMax; ++i)
+		{
+			Boxes.Emplace(UtilsMP4::FMetadataParser::FBoxInfo(md->GetChildBoxType(i), md->GetChildBoxData(i), md->GetChildBoxDataSize(i)));
+		}
+		MediaMetadata = MakeSharedTS<UtilsMP4::FMetadataParser>();
+		if (MediaMetadata->Parse(hdlr, res0, Boxes) == UtilsMP4::FMetadataParser::EResult::Success)
+		{
+			PlayerSessionServices->SendMessageToPlayer(FPlaylistMetadataUpdateMessage::Create(FTimeValue(), MediaMetadata));
+		}
+		else
+		{
+			MediaMetadata.Reset();
+		}
+	}
 	return FErrorDetail();
 }
 
@@ -731,6 +754,7 @@ void FManifestMP4Internal::FTimelineAssetMP4::LimitSegmentDownloadSize(TSharedPt
 		uint32 TrackId = ~0U;
 		uint32 TrackTimeScale = 0;
 		int64 TrackDur = 0;
+		int64 TrackDurSoFar = 0;
 		int64 LastTrackOffset = -1;
 		int64 LastSampleSize = 0;
 		int64 TrackDurationLimit = -1;
@@ -753,17 +777,19 @@ void FManifestMP4Internal::FTimelineAssetMP4::LimitSegmentDownloadSize(TSharedPt
 				}
 				if (TrackId == CurrentTrackIt->GetTrack()->GetID())
 				{
+					TrackDurSoFar = TrackDur;
 					TrackDur += CurrentTrackIt->GetDuration();
 				}
 				int64 CurrentTrackOffset = LastTrackOffset;
+				bool bDurationLimitReached = TrackDurationLimit > 0 && TrackDur > TrackDurationLimit;
 				if (CurrentTrackOffset >= EndOffset ||
 					CurrentTrackOffset - StartOffset >= MaxSegmentSize ||
-					(TrackDurationLimit > 0 && TrackDur > TrackDurationLimit))
+					bDurationLimitReached)
 				{
 					// Limit reached.
 					Request->FileEndOffset = CurrentTrackOffset - 1;
 					Request->SegmentInternalSize = CurrentTrackOffset - StartOffset;
-					Request->SegmentDuration.SetFromND(TrackDur, TrackTimeScale);
+					Request->SegmentDuration.SetFromND(bDurationLimitReached ? TrackDurSoFar : TrackDur, TrackTimeScale);
 					Request->bIsLastSegment = CurrentTrackOffset >= EndOffset;
 					break;
 				}
@@ -1254,7 +1280,7 @@ FErrorDetail FManifestMP4Internal::FRepresentationMP4::CreateFrom(const IParserI
 	}
 
 	// Get bitrate from the average or max bitrate as stored in the track. If not stored it will be 0.
-	Bitrate = InTrack->GetBitrateInfo().AvgBitrate ? InTrack->GetBitrateInfo().AvgBitrate : InTrack->GetBitrateInfo().MaxBitrate;
+	Bitrate = InTrack->GetBitrateInfo().MaxBitrate ? InTrack->GetBitrateInfo().MaxBitrate : InTrack->GetBitrateInfo().AvgBitrate;
 
 	// With no bitrate available we set some defaults. This is mainly to avoid a bitrate of 0 from being surfaced that would prevent
 	// events like the initial bitrate change that needs to transition away from 0 to something real.
@@ -1277,6 +1303,10 @@ FErrorDetail FManifestMP4Internal::FRepresentationMP4::CreateFrom(const IParserI
 				break;
 		}
 	}
+	if (!CodecInformation.GetBitrate())
+	{
+		CodecInformation.SetBitrate(Bitrate);
+	}
 
 	// Not a whole lot that could have gone wrong here.
 	return FErrorDetail();
diff --git a/Engine/Plugins/Media/ElectraPlayer/Source/ElectraPlayerRuntime/Private/Runtime/Player/mp4/ManifestMP4.h b/Engine/Plugins/Media/ElectraPlayer/Source/ElectraPlayerRuntime/Private/Runtime/Player/mp4/ManifestMP4.h
--- a/Engine/Plugins/Media/ElectraPlayer/Source/ElectraPlayerRuntime/Private/Runtime/Player/mp4/ManifestMP4.h
+++ b/Engine/Plugins/Media/ElectraPlayer/Source/ElectraPlayerRuntime/Private/Runtime/Player/mp4/ManifestMP4.h
@@ -12,7 +12,7 @@
 #include "ElectraPlayerPrivate.h"
 #include "Player/PlayerStreamReader.h"
 #include "Utilities/URLParser.h"
-
+#include "Utilities/UtilsMP4.h"
 
 
 namespace Electra
@@ -259,6 +259,8 @@ public:
 		TArray<TSharedPtrTS<FAdaptationSetMP4>>			AudioAdaptationSets;
 		TArray<TSharedPtrTS<FAdaptationSetMP4>>			SubtitleAdaptationSets;
 		FPlayRangeEndInfo								PlayRangeEndInfo;
+		// Metadata from the 'meta' box, if any.
+		TSharedPtrTS<UtilsMP4::FMetadataParser>			MediaMetadata;
 	};
 
 
diff --git a/Engine/Plugins/Media/ElectraPlayer/Source/ElectraPlayerRuntime/Public/IElectraPlayerInterface.h b/Engine/Plugins/Media/ElectraPlayer/Source/ElectraPlayerRuntime/Public/IElectraPlayerInterface.h
--- a/Engine/Plugins/Media/ElectraPlayer/Source/ElectraPlayerRuntime/Public/IElectraPlayerInterface.h
+++ b/Engine/Plugins/Media/ElectraPlayer/Source/ElectraPlayerRuntime/Public/IElectraPlayerInterface.h
@@ -38,6 +38,7 @@ public:
 		PlayListData,
 		LicenseKeyData,
 		PlaystartPosFromSeekPositions,
+		MediaMetadataUpdate,
 		Android_VideoSurface,
 	};
 	virtual Electra::FVariantValue QueryOptions(EOptionType Type, const Electra::FVariantValue& Param = Electra::FVariantValue()) = 0;
diff --git a/Engine/Plugins/Media/ElectraUtil/Source/ElectraBase/Private/UtilsMP4.cpp b/Engine/Plugins/Media/ElectraUtil/Source/ElectraBase/Private/UtilsMP4.cpp
new file mode 100644
--- /dev/null
+++ b/Engine/Plugins/Media/ElectraUtil/Source/ElectraBase/Private/UtilsMP4.cpp
@@ -0,0 +1,683 @@
+// Copyright Epic Games, Inc. All Rights Reserved.
+
+#include "Utilities/UtilsMP4.h"
+#include <Misc/Base64.h>
+
+
+namespace Electra
+{
+
+UtilsMP4::FDataReader::FDataReader(const TConstArrayView<const uint8>& InDataBufferToReadFrom)
+	: DataBufferRef(InDataBufferToReadFrom)
+{
+}
+
+int32 UtilsMP4::FDataReader::GetCurrentOffset() const
+{
+	return CurrentOffset;	
+}
+
+int32 UtilsMP4::FDataReader::GetNumBytesRemaining() const
+{ 
+	return DataBufferRef.Num() - GetCurrentOffset(); 
+}
+
+const uint8* UtilsMP4::FDataReader::GetCurrentDataPointer() const
+{
+	return GetNumBytesRemaining() ? DataBufferRef.GetData() + GetCurrentOffset() : nullptr;
+}
+
+void UtilsMP4::FDataReader::SetCurrentOffset(int32 InNewOffset)
+{
+	check(InNewOffset >= 0 && InNewOffset <= DataBufferRef.Num());
+	if (InNewOffset >= 0 && InNewOffset <= DataBufferRef.Num())
+	{
+		CurrentOffset = InNewOffset;
+	}
+}
+
+bool UtilsMP4::FDataReader::ReadString(FString& OutString, uint16 NumBytes)
+{
+	OutString.Empty();
+	if (NumBytes == 0)
+	{
+		return true;
+	}
+	TArray<uint8> Buf;
+	Buf.AddUninitialized(NumBytes);
+	if (ReadBytes(Buf.GetData(), NumBytes))
+	{
+		// Check for UTF16 BOM
+		if (NumBytes >= 2 && ((Buf[0] == 0xff && Buf[1] == 0xfe) || (Buf[0] == 0xfe && Buf[1] == 0xff)))
+		{
+			// String uses UTF16, which is not supported
+			return false;
+		}
+		FUTF8ToTCHAR cnv((const ANSICHAR*)Buf.GetData(), NumBytes);
+		OutString = FString(cnv.Length(), cnv.Get());
+		return true;
+	}
+	return false;
+}
+
+bool UtilsMP4::FDataReader::ReadStringUTF8(FString& OutString, int32 NumBytes)
+{
+	OutString.Empty();
+	if (NumBytes == 0)
+	{
+		return true;
+	}
+	else if (NumBytes < 0)
+	{
+		NumBytes = GetNumBytesRemaining();
+		check(NumBytes >= 0);
+		if (NumBytes < 0)
+		{
+			return false;
+		}
+	}
+	TArray<uint8> Buf;
+	Buf.AddUninitialized(NumBytes);
+	if (ReadBytes(Buf.GetData(), NumBytes))
+	{
+		FUTF8ToTCHAR cnv((const ANSICHAR*)Buf.GetData(), NumBytes);
+		OutString = FString(cnv.Length(), cnv.Get());
+		return true;
+	}
+	return false;
+}
+
+bool UtilsMP4::FDataReader::ReadStringUTF16(FString& OutString, int32 NumBytes)
+{
+	OutString.Empty();
+	if (NumBytes == 0)
+	{
+		return true;
+	}
+	else if (NumBytes < 0)
+	{
+		NumBytes = GetNumBytesRemaining();
+		check(NumBytes >= 0);
+		if (NumBytes < 0)
+		{
+			return false;
+		}
+	}
+	TArray<uint8> Buf;
+	Buf.AddUninitialized(NumBytes);
+	if (ReadBytes(Buf.GetData(), NumBytes))
+	{
+		check(!"TODO");
+/*
+		FUTF8ToTCHAR cnv((const ANSICHAR*)Buf.GetData(), NumBytes);
+		OutString = FString(cnv.Length(), cnv.Get());
+		return true;
+*/
+	}
+	return false;
+}
+
+bool UtilsMP4::FDataReader::ReadAsNumber(uint64& OutValue, int32 InNumBytes)
+{
+	OutValue = 0;
+	if (InNumBytes < 0 || InNumBytes > 8)
+	{
+		return false;
+	}
+	for(int32 i=0; i<InNumBytes; ++i)
+	{
+		uint8 d;
+		if (!Read(d))
+		{
+			return false;
+		}
+		OutValue = (OutValue << 8) | d;
+	}
+	return true;
+}
+bool UtilsMP4::FDataReader::ReadAsNumber(int64& OutValue, int32 InNumBytes)
+{
+	OutValue = 0;
+	if (InNumBytes < 0 || InNumBytes > 8)
+	{
+		return false;
+	}
+	for(int32 i=0; i<InNumBytes; ++i)
+	{
+		uint8 d;
+		if (!Read(d))
+		{
+			return false;
+		}
+		if (i==0 && d>127)
+		{
+			OutValue = -1;
+		}
+		OutValue = (OutValue << 8) | d;
+	}
+	return true;
+}
+bool UtilsMP4::FDataReader::ReadAsNumber(float& OutValue)
+{
+	uint32 Flt;
+	if (Read(Flt))
+	{
+		OutValue = *reinterpret_cast<float*>(&Flt);
+		return true;
+	}
+	return false;
+}
+bool UtilsMP4::FDataReader::ReadAsNumber(double& OutValue)
+{
+	uint64 Dbl;
+	if (Read(Dbl))
+	{
+		OutValue = *reinterpret_cast<double*>(&Dbl);
+		return true;
+	}
+	return false;
+}
+
+bool UtilsMP4::FDataReader::ReadBytes(void* Buffer, int32 NumBytes)
+{
+	return ReadData(Buffer, NumBytes) == NumBytes;
+}
+
+int32 UtilsMP4::FDataReader::ReadData(void* IntoBuffer, int32 NumBytesToRead)
+{
+	if (NumBytesToRead <= 0)
+	{
+		return 0;
+	}
+	int32 NumAvail = DataBufferRef.Num() - CurrentOffset;
+	if (NumAvail >= NumBytesToRead)
+	{
+		if (IntoBuffer)
+		{
+			FMemory::Memcpy(IntoBuffer, DataBufferRef.GetData() + CurrentOffset, NumBytesToRead);
+		}
+		CurrentOffset += NumBytesToRead;
+		return NumBytesToRead;
+	}
+	return -1;
+}
+
+
+
+/***************************************************************************************************************************************************/
+/***************************************************************************************************************************************************/
+/***************************************************************************************************************************************************/
+
+namespace FMetadataTags
+{
+static const TCHAR * const Title = TEXT("Title");
+static const TCHAR * const Artist = TEXT("Artist");
+static const TCHAR * const Genre = TEXT("Genre");
+static const TCHAR * const Date = TEXT("Date");						// Content creation date
+static const TCHAR * const Description = TEXT("Description");
+static const TCHAR * const LongDescription = TEXT("LongDescription");
+static const TCHAR * const Album = TEXT("Album");
+static const TCHAR * const Encoder = TEXT("Encoder");
+}
+
+UtilsMP4::FMetadataParser::FMetadataParser()
+{
+	NumTotalItems = 0;
+
+	WellKnownItems.Emplace(MakeBoxAtom(0xa9,'n','a','m'), FMetadataTags::Title);
+	WellKnownItems.Emplace(MakeBoxAtom(0xa9,'A','R','T'), FMetadataTags::Artist);
+	WellKnownItems.Emplace(MakeBoxAtom(0xa9,'g','e','n'), FMetadataTags::Genre);
+	WellKnownItems.Emplace(MakeBoxAtom(0xa9,'d','a','y'), FMetadataTags::Date);
+	WellKnownItems.Emplace(MakeBoxAtom('d','e','s','c'), FMetadataTags::Description);
+	WellKnownItems.Emplace(MakeBoxAtom('l','d','e','s'), FMetadataTags::LongDescription);
+	WellKnownItems.Emplace(MakeBoxAtom(0xa9,'a','l','b'), FMetadataTags::Album);
+	WellKnownItems.Emplace(MakeBoxAtom(0xa9,'t','o','o'), FMetadataTags::Encoder);
+}
+
+FString UtilsMP4::FMetadataParser::PrintableBoxAtom(const uint32 InAtom)
+{
+	FString Out;
+	// Not so much just printable as alphanumeric.
+	for(uint32 i=0, Atom=InAtom; i<4; ++i, Atom<<=8)
+	{
+		int32 v = Atom >> 24;
+		if ((v>='A' && v<='Z') || (v>='a' && v<='z') || (v>='0' && v<='9') || v=='_')
+		{
+			Out.AppendChar(v);
+		}
+		else
+		{
+			// Not alphanumeric, return it as a hex string.
+			return FString::Printf(TEXT("%08x"), InAtom);
+		}
+	}
+	return Out;
+}
+
+UtilsMP4::FMetadataParser::EResult UtilsMP4::FMetadataParser::Parse(uint32 InHandler, uint32 InHandlerReserved0, const TArray<UtilsMP4::FMetadataParser::FBoxInfo>& InBoxes)
+{
+	// We only support the Apple iTunes metadata at the moment.
+	if (InHandler != MakeBoxAtom('m','d','i','r'))
+	{
+		return UtilsMP4::FMetadataParser::EResult::NotSupported;
+	}
+	/*
+		As per ISO/IEC 14496-12:2015 section 8.11.1.2 the 'meta' box must contain a handler ('hdlr'), followed by nothing but optional boxes.
+		Of these optional boxes the ones recognized and defined by the document are:
+		'pitm' - 8.11.4 Primary Item Box
+		'dinf' - 8.7.1 Data Information Box
+		'iloc' - 8.11.3 The Item Location Box
+		'ipro' - 8.11.5 Item Protection Box
+		'iinf' - 8.11.6 Item Information Box
+		 ....  - IPMPControlBox (8.12 Support for Protected Streams)
+		'iref' - 8.11.12 Item Reference Box
+		'idat' - 8.11.11 Item Data Box
+
+		followed by any number of other boxes.
+
+		Presently the boxes mentioned above are not processed since they are apparently not used by any of the tools
+		that allow insertion of metadata. It appears that the way iTunes handles metadata has become the de facto standard
+		in which there is
+		    'meta'
+			   'hdlr'
+			   'ilst'
+			     ....               // See: https://developer.apple.com/library/archive/documentation/QuickTime/QTFF/Metadata/Metadata.html#//apple_ref/doc/uid/TP40000939-CH1-SW26
+				   'data'			// See: https://developer.apple.com/library/archive/documentation/QuickTime/QTFF/Metadata/Metadata.html#//apple_ref/doc/uid/TP40000939-CH1-SW27
+	*/
+
+	// By some definition the items are contained inside an 'ilst' box, so we look for that one.
+	bool bHaveILST = false;
+	for(int32 i=0; i<InBoxes.Num(); ++i)
+	{
+		if (InBoxes[i].Type == MakeBoxAtom('i','l','s','t'))
+		{
+			// We will actually parse all 'ilst' boxes if for some reason the metadata has
+			// been scattered across several.
+			bHaveILST = true;
+			Parse(InBoxes[i]);
+		}
+	}
+	return bHaveILST ? UtilsMP4::FMetadataParser::EResult::Success : UtilsMP4::FMetadataParser::EResult::MissingBox;
+}
+
+void UtilsMP4::FMetadataParser::Parse(const UtilsMP4::FMetadataParser::FBoxInfo& InBox)
+{
+	if (!InBox.Data || !InBox.Size)
+	{
+		return;
+	}
+	
+	FDataReader dr(MakeArrayView<const uint8>(static_cast<const uint8*>(InBox.Data), InBox.Size));
+	while(dr.GetNumBytesRemaining() > 8)
+	{
+		int32 BoxSize;
+		uint32 BoxType;
+		dr.Read(BoxSize);
+		dr.Read(BoxType);
+		if (WellKnownItems.Contains(BoxType))
+		{
+			ParseBoxDataList(WellKnownItems[BoxType], dr.GetCurrentDataPointer(), BoxSize-8);
+		}
+		else switch(BoxType)
+		{
+			case MakeBoxAtom('-','-','-','-'):
+			{
+				ParseBoxDataiTunes(dr.GetCurrentDataPointer(), BoxSize-8);
+				break;
+			}
+			default:
+			{
+				ParseBoxDataList(FString(TEXT("qt."))+PrintableBoxAtom(BoxType), dr.GetCurrentDataPointer(), BoxSize-8);
+				break;
+			}
+		}
+		dr.ReadBytes(nullptr, BoxSize - 8);
+	}
+}
+
+
+void UtilsMP4::FMetadataParser::ParseBoxDataList(const FString& AsCategory, const uint8* InBoxData, uint32 InBoxSize)
+{
+	if (!InBoxData || InBoxSize <= 8)
+	{
+		return;
+	}
+
+	FDataReader dr(MakeArrayView<const uint8>(static_cast<const uint8*>(InBoxData), InBoxSize));
+
+#define RETURN_IF_READERROR(expr)	\
+	if (!(expr))					\
+	{								\
+		return;						\
+	}								\
+
+	for(int32 CurrentOffset=dr.GetCurrentOffset(), BoxSize=0; dr.GetNumBytesRemaining() > 8; dr.SetCurrentOffset(CurrentOffset), dr.ReadBytes(nullptr, BoxSize - 8))
+	{
+		uint32 BoxType;
+		RETURN_IF_READERROR(dr.Read(BoxSize));
+		RETURN_IF_READERROR(dr.Read(BoxType));
+		CurrentOffset = dr.GetCurrentOffset();
+		if (BoxType != MakeBoxAtom('d','a','t','a'))
+		{
+			continue;
+		}
+		uint32 TypeIndicatorAndWellKnownType = 0;
+		RETURN_IF_READERROR(dr.Read(TypeIndicatorAndWellKnownType));
+		if ((TypeIndicatorAndWellKnownType >> 24) != 0)
+		{
+			continue;
+		}
+		uint16 CountryIndicator=0;
+		RETURN_IF_READERROR(dr.Read(CountryIndicator));
+		uint16 LanguageIndicator=0;
+		RETURN_IF_READERROR(dr.Read(LanguageIndicator));
+		int32 NumRemainingDataBytes = BoxSize - 8 - (dr.GetCurrentOffset() - CurrentOffset);
+
+		FItem Item;
+
+		// We do not handle the country indicator at the moment.
+		// Likewise, the language indicator must be a directly specified language code since we do not handle
+		// language tables.
+		if (LanguageIndicator >= 1 && LanguageIndicator <= 255)
+		{
+			// This would be an index into the language table (which is a 'lang' box somewhere under the 'meta' box).
+			continue;
+		}
+		if (LanguageIndicator)
+		{
+			Item.Language[0] = (char)(0x60 + ((LanguageIndicator & 0x7c00) >> 10));
+			Item.Language[1] = (char)(0x60 + ((LanguageIndicator & 0x03e0) >> 5));
+			Item.Language[2] = (char)(0x60 + (LanguageIndicator & 0x001f));
+		}
+
+		bool bSet = true;
+		int32 WKT = (int32)(TypeIndicatorAndWellKnownType & 0x00ffffffU);
+		switch(WKT)
+		{
+			// UTF-8 string
+			case 1:
+			{
+				FString String;
+				if ((bSet = dr.ReadStringUTF8(String, NumRemainingDataBytes)) == true)
+				{
+					Item.Value = String;
+					Item.Type = 1;
+				}
+				break;
+			}
+/*
+			// UTF-16 string
+			case 2:
+			{
+				FString String;
+				if ((bSet = dr.ReadStringUTF16(String, NumRemainingDataBytes)) == true)
+				{
+					Item.Value = String;
+					Item.Type = 1;
+				}
+				break;
+			}
+*/
+			case 13:	// JPEG image
+			case 14:	// PNG image
+			case 27:	// BMP image
+			{
+				TArray<uint8> Image(MakeArrayView<const uint8>(dr.GetCurrentDataPointer(), NumRemainingDataBytes));
+				Item.Value = Image;
+				Item.Type = WKT;
+				break;
+			}
+
+			// Float
+			case 23:
+			{
+				float Flt;
+				if ((bSet = dr.ReadAsNumber(Flt)) == true)
+				{
+					// Set as double!
+					Item.Value = (double)Flt;
+					Item.Type = 24;
+				}
+				break;
+			}
+			// Double
+			case 24:
+			{
+				double Dbl;
+				if ((bSet = dr.ReadAsNumber(Dbl)) == true)
+				{
+					Item.Value = Dbl;
+					Item.Type = 24;
+				}
+				break;
+			}
+			// Signed integers
+			case 65:	// int8
+			case 66:	// int16
+			case 67:	// int32
+			case 74:	// int64
+			case 21:	// 1-4 byte signed integer
+			{
+				int64 Number;
+				if ((bSet = dr.ReadAsNumber(Number, WKT==65?1:WKT==66?2:WKT==67?4:WKT==74?8:NumRemainingDataBytes)) == true)
+				{
+					Item.Value = Number;
+					Item.Type = 74;
+				}
+				break;
+			}
+			// Unsigned integers
+			case 75:	// uint8
+			case 76:	// uint16
+			case 77:	// uint32
+			case 78:	// uint64
+			case 22:	// 1-4 byte unsigned integer
+			{
+				uint64 Number;
+				if ((bSet = dr.ReadAsNumber(Number, WKT==75?1:WKT==76?2:WKT==77?4:WKT==78?8:NumRemainingDataBytes)) == true)
+				{
+					Item.Value = Number;
+					Item.Type = 78;
+				}
+				break;
+			}
+			
+			// Not handled here.
+			default:
+			{
+				bSet = false;
+				break;
+			}
+		}
+
+		if (bSet)
+		{
+			++NumTotalItems;
+			Items.FindOrAdd(AsCategory).Emplace(MoveTemp(Item));
+		}
+	}
+#undef RETURN_IF_READERROR
+}
+
+void UtilsMP4::FMetadataParser::ParseBoxDataiTunes(const uint8* InBoxData, uint32 InBoxSize)
+{
+	if (!InBoxData || InBoxSize <= 8)
+	{
+		return;
+	}
+
+	FDataReader dr(MakeArrayView<const uint8>(static_cast<const uint8*>(InBoxData), InBoxSize));
+
+#define RETURN_IF_READERROR(expr)	\
+	if (!(expr))					\
+	{								\
+		return;						\
+	}								\
+
+	FString Name;
+	for(int32 CurrentOffset=dr.GetCurrentOffset(), BoxSize=0; dr.GetNumBytesRemaining() > 8; dr.SetCurrentOffset(CurrentOffset), dr.ReadBytes(nullptr, BoxSize - 8))
+	{
+		uint32 BoxType;
+		RETURN_IF_READERROR(dr.Read(BoxSize));
+		RETURN_IF_READERROR(dr.Read(BoxType));
+		CurrentOffset = dr.GetCurrentOffset();
+		if (BoxType == MakeBoxAtom('m','e','a','n'))
+		{
+			// 4 reserved bytes follow
+			RETURN_IF_READERROR(dr.ReadBytes(nullptr, 4));
+			FString Meaning;
+			int32 NumRemainingDataBytes = BoxSize - 8 - (dr.GetCurrentOffset() - CurrentOffset);
+			RETURN_IF_READERROR(dr.ReadStringUTF8(Meaning, NumRemainingDataBytes));
+			if (!Meaning.Equals(TEXT("com.apple.iTunes")))
+			{
+				return;
+			}
+		}
+		else if (BoxType == MakeBoxAtom('n','a','m','e'))
+		{
+			// 4 reserved bytes follow
+			RETURN_IF_READERROR(dr.ReadBytes(nullptr, 4));
+			int32 NumRemainingDataBytes = BoxSize - 8 - (dr.GetCurrentOffset() - CurrentOffset);
+			RETURN_IF_READERROR(dr.ReadStringUTF8(Name, NumRemainingDataBytes));
+		}
+		else if (BoxType == MakeBoxAtom('d','a','t','a'))
+		{
+			ParseBoxDataList(FString(TEXT("iTunes."))+Name, dr.GetCurrentDataPointer()-8, BoxSize);
+		}
+	}
+#undef RETURN_IF_READERROR
+}
+
+bool UtilsMP4::FMetadataParser::IsDifferentFrom(const UtilsMP4::FMetadataParser& Other)
+{
+	// Different number of items or map entries means the metadata can't be identical.
+	if (NumTotalItems != Other.NumTotalItems || Items.Num() != Other.Items.Num())
+	{
+		return true;
+	}
+
+	// Check if the maps keys or items are different.
+	for(auto& Key : Items)
+	{
+		const TArray<FItem>* OtherKeyItems = Other.Items.Find(Key.Key);
+		if (!OtherKeyItems || Key.Value.Num() != OtherKeyItems->Num())
+		{
+			return true;
+		}
+		// Expensive item by item test :-(
+		for(auto& Item : Key.Value)
+		{
+			if (!OtherKeyItems->Contains(Item))
+			{
+				return true;
+			}
+		}
+	}
+	return false;
+}
+
+
+FString UtilsMP4::FMetadataParser::GetAsJSON() const
+{
+	FString JSON(TEXT("{"));
+	bool bFirstItem = true;
+	for(auto& Item : Items)
+	{
+		bool bFirstItemValue = true;
+		const TArray<FItem>& List=Item.Value;
+		FString ItemJSON;
+		for(int32 i=0; i<List.Num(); ++i)
+		{
+			const FItem& it = List[i];
+			FString ItemValue = it.ToJSONValue();
+			if (ItemValue.Len())
+			{
+				if (!bFirstItemValue)
+				{
+					ItemJSON.AppendChar(TCHAR(','));
+				}
+				bFirstItemValue = false;
+				FString Lang(3, it.Language);
+				ItemJSON.Append(FString::Printf(TEXT("\"%s\":"), *Lang));
+				ItemJSON.Append(MoveTemp(ItemValue));
+			}
+		}
+		if (ItemJSON.Len())
+		{
+			if (!bFirstItem)
+			{
+				JSON.AppendChar(TCHAR(','));
+			}
+			bFirstItem = false;
+
+			JSON.Append(FString::Printf(TEXT("\"%s\":{"), *Item.Key));
+			JSON.Append(MoveTemp(ItemJSON));
+			JSON.Append(TEXT("}"));
+		}
+	}
+	JSON.Append(TEXT("}"));
+	return JSON;
+}
+
+FString UtilsMP4::FMetadataParser::FItem::ToJSONValue() const
+{
+	switch(Type)
+	{
+		// String
+		case 1:
+		{
+			return FString::Printf(TEXT("\"%s\""), *Value.GetValue<FString>().ReplaceCharWithEscapedChar(&CharsToEscapeInJSON));
+		}
+		// Double
+		case 24:
+		{
+			return FString::Printf(TEXT("%f"), Value.GetValue<double>());
+		}
+		// Signed integer
+		case 74:
+		{
+			int64 v = Value.GetValue<int64>();
+			if (v >= -(1LL<<53) && v <= (1LL<<53)-1)
+			{
+				return FString::Printf(TEXT("%lld"), (long long int)v);
+			}
+		}
+		// Unsigned integer
+		case 78:
+		{
+			uint64 v = Value.GetValue<uint64>();
+			if (v <= (1LL<<53)-1)
+			{
+				return FString::Printf(TEXT("%lld"), (long long int)v);
+			}
+		}
+		// JPEG image
+		case 13:
+		{
+			FString js(TEXT("{\"mimetype\":\"image/jpeg\",\"data\":\""));
+			js.Append(FBase64::Encode(Value.GetValue<TArray<uint8>>()));
+			js.Append(TEXT("\"}"));
+			return js;
+		}
+		// PNG image
+		case 14:
+		{
+			FString js(TEXT("{\"mimetype\":\"image/png\",\"data\":\""));
+			js.Append(FBase64::Encode(Value.GetValue<TArray<uint8>>()));
+			js.Append(TEXT("\"}"));
+			return js;
+		}
+		default:
+		{
+			break;
+		}
+	}
+	return FString();
+}
+TArray<TCHAR> UtilsMP4::FMetadataParser::FItem::CharsToEscapeInJSON = { TCHAR('\n'), TCHAR('\r'), TCHAR('\t'), TCHAR('\\'), TCHAR('\"') };
+
+
+
+}
+
+
diff --git a/Engine/Plugins/Media/ElectraUtil/Source/ElectraBase/Public/Utilities/UtilsMP4.h b/Engine/Plugins/Media/ElectraUtil/Source/ElectraBase/Public/Utilities/UtilsMP4.h
new file mode 100644
--- /dev/null
+++ b/Engine/Plugins/Media/ElectraUtil/Source/ElectraBase/Public/Utilities/UtilsMP4.h
@@ -0,0 +1,150 @@
+// Copyright Epic Games, Inc. All Rights Reserved.
+
+#pragma once
+
+#include <CoreMinimal.h>
+#include <Misc/Variant.h>
+
+namespace Electra
+{
+
+class ELECTRABASE_API UtilsMP4
+{
+public:
+
+#if !PLATFORM_LITTLE_ENDIAN
+	static inline uint8 GetFromBigEndian(uint8 value)		{ return value; }
+	static inline int8 GetFromBigEndian(int8 value)			{ return value; }
+	static inline uint16 GetFromBigEndian(uint16 value)		{ return value; }
+	static inline int16 GetFromBigEndian(int16 value)		{ return value; }
+	static inline int32 GetFromBigEndian(int32 value)		{ return value; }
+	static inline uint32 GetFromBigEndian(uint32 value)		{ return value; }
+	static inline int64 GetFromBigEndian(int64 value)		{ return value; }
+	static inline uint64 GetFromBigEndian(uint64 value)		{ return value; }
+#else
+	static inline uint16 EndianSwap(uint16 value)			{ return (value >> 8) | (value << 8); }
+	static inline int16 EndianSwap(int16 value)				{ return int16(EndianSwap(uint16(value))); }
+	static inline uint32 EndianSwap(uint32 value)			{ return (value << 24) | ((value & 0xff00) << 8) | ((value >> 8) & 0xff00) | (value >> 24); }
+	static inline int32 EndianSwap(int32 value)				{ return int32(EndianSwap(uint32(value))); }
+	static inline uint64 EndianSwap(uint64 value)			{ return (uint64(EndianSwap(uint32(value & 0xffffffffU))) << 32) | uint64(EndianSwap(uint32(value >> 32))); }
+	static inline int64 EndianSwap(int64 value)				{ return int64(EndianSwap(uint64(value)));}
+	static inline uint8 GetFromBigEndian(uint8 value)		{ return value; }
+	static inline int8 GetFromBigEndian(int8 value)			{ return value; }
+	static inline uint16 GetFromBigEndian(uint16 value)		{ return EndianSwap(value); }
+	static inline int16 GetFromBigEndian(int16 value)		{ return EndianSwap(value); }
+	static inline int32 GetFromBigEndian(int32 value)		{ return EndianSwap(value); }
+	static inline uint32 GetFromBigEndian(uint32 value)		{ return EndianSwap(value); }
+	static inline int64 GetFromBigEndian(int64 value)		{ return EndianSwap(value); }
+	static inline uint64 GetFromBigEndian(uint64 value)		{ return EndianSwap(value); }
+#endif
+
+	static constexpr uint32 MakeBoxAtom(const uint8 A, const uint8 B, const uint8 C, const uint8 D)
+	{
+		return (static_cast<uint32>(A) << 24) | (static_cast<uint32>(B) << 16) | (static_cast<uint32>(C) << 8) | static_cast<uint32>(D);
+	}
+
+
+	class ELECTRABASE_API FDataReader
+	{
+	public:
+		FDataReader(const TConstArrayView<const uint8>& InDataBufferToReadFrom);
+
+		int32 GetCurrentOffset() const;
+		int32 GetNumBytesRemaining() const;
+		const uint8* GetCurrentDataPointer() const;
+		void SetCurrentOffset(int32 InNewOffset);
+
+		template <typename T>
+		bool Read(T& value)
+		{
+			T Temp = 0;
+			int64 NumRead = ReadData(&Temp, sizeof(T));
+			if (NumRead == sizeof(T))
+			{
+				value = ValueFromBigEndian(Temp);
+				return true;
+			}
+			return false;
+		}
+
+		bool ReadString(FString& OutString, uint16 NumBytes);
+		bool ReadStringUTF8(FString& OutString, int32 NumBytes);
+		bool ReadStringUTF16(FString& OutString, int32 NumBytes);
+		bool ReadBytes(void* Buffer, int32 NumBytes);
+		bool ReadAsNumber(int64& OutValue, int32 InNumBytes);
+		bool ReadAsNumber(uint64& OutValue, int32 InNumBytes);
+		bool ReadAsNumber(float& OutValue);
+		bool ReadAsNumber(double& OutValue);
+
+	private:
+		template <typename T>
+		T ValueFromBigEndian(const T value)
+		{
+			return GetFromBigEndian(value);
+		}
+		int32 ReadData(void* IntoBuffer, int32 NumBytesToRead);
+
+		const TConstArrayView<const uint8>& DataBufferRef;
+		int32 CurrentOffset = 0;
+	};
+
+
+	/**
+	 * This class parses metadata embedded in an mp4 / ISO14496-12 file.
+	 * Presently only the structure as used and defined by Apple iTunes is supported.
+	 */
+	class ELECTRABASE_API FMetadataParser
+	{
+	public:
+		FMetadataParser();
+		enum class EResult
+		{
+			Success,
+			NotSupported,
+			MissingBox
+		};
+		struct FBoxInfo
+		{
+			FBoxInfo(uint32 InType, const void* InData, uint32 InSize) : Type(InType), Data(InData), Size(InSize)
+			{}
+			uint32 Type = 0;
+			const void* Data = nullptr;
+			uint32 Size = 0;
+		};
+		EResult Parse(uint32 InHandler, uint32 InHandlerReserved0, const TArray<FBoxInfo>& InBoxes);
+		bool IsDifferentFrom(const FMetadataParser& Other);
+		FString GetAsJSON() const;
+	private:
+		FString PrintableBoxAtom(const uint32 InAtom);
+		void Parse(const FBoxInfo& InBox);
+		void ParseBoxDataList(const FString& AsCategory, const uint8* InBoxData, uint32 InBoxSize);
+		void ParseBoxDataiTunes(const uint8* InBoxData, uint32 InBoxSize);
+
+		struct FItem
+		{
+			char Language[3] { 0 };			// ISO 639-2; if not set (all zero) the default entry for all languages
+			int32 Type = 0;					// Well-known data type (see Quicktime reference)
+			FVariant Value;
+			FString ToJSONValue() const;
+			static TArray<TCHAR> CharsToEscapeInJSON;
+			bool operator != (const FItem& Other) const
+			{
+				return Type != Other.Type ||
+						FMemory::Memcmp(Language, Other.Language, sizeof(Language)) != 0 ||
+						Value != Other.Value;
+			}
+			bool operator == (const FItem& Other) const
+			{
+				return !(*this != Other);
+			}
+		};
+
+		TMap<uint32, FString> WellKnownItems;
+		TMap<FString, TArray<FItem>> Items;
+		uint32 NumTotalItems = 0;
+	};
+
+
+};
+
+} // namespace Electra
