ElectraUtil: replaced TArrayView<> with plain pointers due to some crashes on optimized clang builds

#rnx

diff --git a/Engine/Plugins/Media/ElectraUtil/Source/ElectraBase/Private/UtilsMP4.cpp b/Engine/Plugins/Media/ElectraUtil/Source/ElectraBase/Private/UtilsMP4.cpp
--- a/Engine/Plugins/Media/ElectraUtil/Source/ElectraBase/Private/UtilsMP4.cpp
+++ b/Engine/Plugins/Media/ElectraUtil/Source/ElectraBase/Private/UtilsMP4.cpp
@@ -7,8 +7,8 @@
 namespace Electra
 {
 
-UtilsMP4::FDataReader::FDataReader(const TConstArrayView<const uint8>& InDataBufferToReadFrom)
-	: DataBufferRef(InDataBufferToReadFrom)
+UtilsMP4::FDataReader::FDataReader(const void* InDataPtr, int32 InDataSize)
+	: DataPtr((const uint8*)InDataPtr), DataSize(InDataSize), CurrentOffset(0)
 {
 }
 
@@ -19,18 +19,18 @@ int32 UtilsMP4::FDataReader::GetCurrentOffset() const
 
 int32 UtilsMP4::FDataReader::GetNumBytesRemaining() const
 { 
-	return DataBufferRef.Num() - GetCurrentOffset(); 
+	return DataSize - GetCurrentOffset(); 
 }
 
 const uint8* UtilsMP4::FDataReader::GetCurrentDataPointer() const
 {
-	return GetNumBytesRemaining() ? DataBufferRef.GetData() + GetCurrentOffset() : nullptr;
+	return GetNumBytesRemaining() ? DataPtr + GetCurrentOffset() : nullptr;
 }
 
 void UtilsMP4::FDataReader::SetCurrentOffset(int32 InNewOffset)
 {
-	check(InNewOffset >= 0 && InNewOffset <= DataBufferRef.Num());
-	if (InNewOffset >= 0 && InNewOffset <= DataBufferRef.Num())
+	check(InNewOffset >= 0 && InNewOffset <= DataSize);
+	if (InNewOffset >= 0 && InNewOffset <= DataSize)
 	{
 		CurrentOffset = InNewOffset;
 	}
@@ -189,12 +189,12 @@ int32 UtilsMP4::FDataReader::ReadData(void* IntoBuffer, int32 NumBytesToRead)
 	{
 		return 0;
 	}
-	int32 NumAvail = DataBufferRef.Num() - CurrentOffset;
+	int32 NumAvail = DataSize - CurrentOffset;
 	if (NumAvail >= NumBytesToRead)
 	{
 		if (IntoBuffer)
 		{
-			FMemory::Memcpy(IntoBuffer, DataBufferRef.GetData() + CurrentOffset, NumBytesToRead);
+			FMemory::Memcpy(IntoBuffer, DataPtr + CurrentOffset, NumBytesToRead);
 		}
 		CurrentOffset += NumBytesToRead;
 		return NumBytesToRead;
@@ -307,7 +307,7 @@ void UtilsMP4::FMetadataParser::Parse(const UtilsMP4::FMetadataParser::FBoxInfo&
 		return;
 	}
 	
-	FDataReader dr(MakeArrayView<const uint8>(static_cast<const uint8*>(InBox.Data), InBox.Size));
+	FDataReader dr(InBox.Data, InBox.Size);
 	while(dr.GetNumBytesRemaining() > 8)
 	{
 		int32 BoxSize;
@@ -343,7 +343,7 @@ void UtilsMP4::FMetadataParser::ParseBoxDataList(const FString& AsCategory, cons
 		return;
 	}
 
-	FDataReader dr(MakeArrayView<const uint8>(static_cast<const uint8*>(InBoxData), InBoxSize));
+	FDataReader dr(InBoxData, InBoxSize);
 
 #define RETURN_IF_READERROR(expr)	\
 	if (!(expr))					\
@@ -422,8 +422,8 @@ void UtilsMP4::FMetadataParser::ParseBoxDataList(const FString& AsCategory, cons
 			case 14:	// PNG image
 			case 27:	// BMP image
 			{
-				TArray<uint8> Image(MakeArrayView<const uint8>(dr.GetCurrentDataPointer(), NumRemainingDataBytes));
-				Item.Value = Image;
+				TArray<uint8> Image(dr.GetCurrentDataPointer(), NumRemainingDataBytes);
+				Item.Value = MoveTemp(Image);
 				Item.Type = WKT;
 				break;
 			}
@@ -506,7 +506,7 @@ void UtilsMP4::FMetadataParser::ParseBoxDataiTunes(const uint8* InBoxData, uint3
 		return;
 	}
 
-	FDataReader dr(MakeArrayView<const uint8>(static_cast<const uint8*>(InBoxData), InBoxSize));
+	FDataReader dr(InBoxData, InBoxSize);
 
 #define RETURN_IF_READERROR(expr)	\
 	if (!(expr))					\
diff --git a/Engine/Plugins/Media/ElectraUtil/Source/ElectraBase/Public/Utilities/UtilsMP4.h b/Engine/Plugins/Media/ElectraUtil/Source/ElectraBase/Public/Utilities/UtilsMP4.h
--- a/Engine/Plugins/Media/ElectraUtil/Source/ElectraBase/Public/Utilities/UtilsMP4.h
+++ b/Engine/Plugins/Media/ElectraUtil/Source/ElectraBase/Public/Utilities/UtilsMP4.h
@@ -47,7 +47,7 @@ public:
 	class ELECTRABASE_API FDataReader
 	{
 	public:
-		FDataReader(const TConstArrayView<const uint8>& InDataBufferToReadFrom);
+		FDataReader(const void* InDataPtr, int32 InDataSize);
 
 		int32 GetCurrentOffset() const;
 		int32 GetNumBytesRemaining() const;
@@ -84,7 +84,8 @@ public:
 		}
 		int32 ReadData(void* IntoBuffer, int32 NumBytesToRead);
 
-		const TConstArrayView<const uint8>& DataBufferRef;
+		const uint8* DataPtr = nullptr;
+		int32 DataSize = 0;
 		int32 CurrentOffset = 0;
 	};
 
