Adding HDR playback capability for MFW/Electra to Apple platform

diff --git a/Engine/Plugins/Media/ElectraPlayer/Source/ElectraPlayerRuntime/Private/Runtime/Decoder/Apple/VideoDecoderH265_Apple.cpp b/Engine/Plugins/Media/ElectraPlayer/Source/ElectraPlayerRuntime/Private/Runtime/Decoder/Apple/VideoDecoderH265_Apple.cpp
--- a/Engine/Plugins/Media/ElectraPlayer/Source/ElectraPlayerRuntime/Private/Runtime/Decoder/Apple/VideoDecoderH265_Apple.cpp
+++ b/Engine/Plugins/Media/ElectraPlayer/Source/ElectraPlayerRuntime/Private/Runtime/Decoder/Apple/VideoDecoderH265_Apple.cpp
@@ -765,24 +765,37 @@ bool FVideoDecoderH265::InternalDecoderCreate(CMFormatDescriptionRef InputFormat
 
 	// Output image format configuration
 	CFMutableDictionaryRef OutputImageFormat = CFDictionaryCreateMutable(nullptr, 3, &kCFTypeDictionaryKeyCallBacks, &kCFTypeDictionaryValueCallBacks);
+
+	// Allow for an array of output formats, so the decoder can deliver SDR and HDR content in suitable buffers
+	const int NumFmts = 4;
+	int pxfmt[NumFmts] = {
+		kCVPixelFormatType_420YpCbCr8BiPlanarVideoRange,
+		kCVPixelFormatType_420YpCbCr8BiPlanarFullRange,
+		kCVPixelFormatType_420YpCbCr10BiPlanarVideoRange,
+		kCVPixelFormatType_420YpCbCr10BiPlanarFullRange
+		};
+	CFNumberRef PixelFormats[NumFmts] = {
+		CFNumberCreate(kCFAllocatorDefault, kCFNumberSInt32Type, &pxfmt[0]),
+		CFNumberCreate(kCFAllocatorDefault, kCFNumberSInt32Type, &pxfmt[1]),
+		CFNumberCreate(kCFAllocatorDefault, kCFNumberSInt32Type, &pxfmt[2]),
+		CFNumberCreate(kCFAllocatorDefault, kCFNumberSInt32Type, &pxfmt[3])
+		};
+	CFArrayRef PixFmtArray = CFArrayCreate(kCFAllocatorDefault, (const void**)PixelFormats, NumFmts, &kCFTypeArrayCallBacks);
+
+	CFDictionarySetValue(OutputImageFormat, kCVPixelBufferPixelFormatTypeKey, PixFmtArray);
+
 	// Choice of: kCVPixelBufferOpenGLCompatibilityKey (all)  kCVPixelBufferOpenGLESCompatibilityKey (iOS only)   kCVPixelBufferMetalCompatibilityKey (all)
-#if PLATFORM_MAC
-	int pxfmt = kCVPixelFormatType_32BGRA;
-	CFNumberRef PixelFormat = CFNumberCreate(kCFAllocatorDefault, kCFNumberSInt32Type, &pxfmt);
-	CFDictionarySetValue(OutputImageFormat, kCVPixelBufferPixelFormatTypeKey, PixelFormat);
-	CFRelease(PixelFormat);
 	CFDictionarySetValue(OutputImageFormat, kCVPixelBufferMetalCompatibilityKey, kCFBooleanTrue);
-#elif PLATFORM_IOS
+#if PLATFORM_IOS
 	CFDictionarySetValue(OutputImageFormat, kCVPixelBufferOpenGLESCompatibilityKey, kCFBooleanFalse);
-	CFDictionarySetValue(OutputImageFormat, kCVPixelBufferMetalCompatibilityKey, kCFBooleanTrue);
-	int pxfmt = kCVPixelFormatType_32BGRA;
-	CFNumberRef PixelFormat = CFNumberCreate(kCFAllocatorDefault, kCFNumberSInt32Type, &pxfmt);
-	CFDictionarySetValue(OutputImageFormat, kCVPixelBufferPixelFormatTypeKey, PixelFormat);
-	CFRelease(PixelFormat);
-#else
-	#error "Should not get here. Check platform checks at the top of the file."
 #endif
 
+	CFRelease(PixFmtArray);
+	for (int Idx = 0; Idx < NumFmts; ++Idx)
+	{
+		CFRelease(PixelFormats[Idx]);
+	}
+
 	// Session configuration
 	CFMutableDictionaryRef SessionConfiguration = CFDictionaryCreateMutable(nullptr, 1, &kCFTypeDictionaryKeyCallBacks, &kCFTypeDictionaryValueCallBacks);
 	// Ask for hardware decoding
diff --git a/Engine/Plugins/Media/ElectraUtil/Source/ElectraSamples/ElectraSamples.Build.cs b/Engine/Plugins/Media/ElectraUtil/Source/ElectraSamples/ElectraSamples.Build.cs
--- a/Engine/Plugins/Media/ElectraUtil/Source/ElectraSamples/ElectraSamples.Build.cs
+++ b/Engine/Plugins/Media/ElectraUtil/Source/ElectraSamples/ElectraSamples.Build.cs
@@ -23,6 +23,7 @@ namespace UnrealBuildTool.Rules
 					"RenderCore",
 					"RHI",
 					"ElectraBase",
+					"ColorManagement",
 				});
 
 			PrivateIncludePathModuleNames.AddRange(
diff --git a/Engine/Plugins/Media/ElectraUtil/Source/ElectraSamples/Private/ElectraTextureSampleAVF.cpp b/Engine/Plugins/Media/ElectraUtil/Source/ElectraSamples/Private/ElectraTextureSampleAVF.cpp
--- a/Engine/Plugins/Media/ElectraUtil/Source/ElectraSamples/Private/ElectraTextureSampleAVF.cpp
+++ b/Engine/Plugins/Media/ElectraUtil/Source/ElectraSamples/Private/ElectraTextureSampleAVF.cpp
@@ -150,7 +150,7 @@ public:
 } // namespace anonymous
 
 
-void FElectraMediaTexConvApple::ConvertTexture(FTexture2DRHIRef & InDstTexture, CVImageBufferRef InImageBufferRef)
+void FElectraMediaTexConvApple::ConvertTexture(FTexture2DRHIRef & InDstTexture, CVImageBufferRef InImageBufferRef, bool bFullRange, EMediaTextureSampleFormat Format, const FMatrix44f& YUVMtx, const FMatrix44f& GamutToXYZMtx, UE::Color::EEncoding EncodingType)
 {
 	check(IsInRenderingThread());
 
@@ -176,18 +176,12 @@ void FElectraMediaTexConvApple::ConvertTexture(FTexture2DRHIRef & InDstTexture,
 
 		if (CVPixelBufferIsPlanar(InImageBufferRef))
 		{
-			FMatrix const* ColorTransform;
-// FIXME: where to get fullrange or not from?
-// FIXME: need to make sure this is properly handling sRGB vs. RGB (as we get a texture from the outside!)
-// NOTE: on all HW we tested to far, we never seem to get here for H.264 material...
-			if(0)//bFullRange)
-			{
-				ColorTransform = &MediaShaders::YuvToRgbRec709Scaled;
-			}
-			else
-			{
-				ColorTransform = &MediaShaders::YuvToRgbRec709Unscaled;
-			}
+			// Planar data: YbCrCb 420 etc. (NV12 / P010)
+
+			bool bIs8Bit = (Format == EMediaTextureSampleFormat::CharNV12);
+
+			const FMatrix* ColorTransform = bFullRange ? &MediaShaders::YuvToRgbRec709Unscaled : &MediaShaders::YuvToRgbRec709Scaled;
+			FVector Off = bFullRange ? MediaShaders::YUVOffsetNoScale8bits : MediaShaders::YUVOffset8bits;
 
 			// Expecting BiPlanar kCVPixelFormatType_420YpCbCr8BiPlanar Full/Video
 			check(CVPixelBufferGetPlaneCount(InImageBufferRef) == 2);
@@ -196,7 +190,7 @@ void FElectraMediaTexConvApple::ConvertTexture(FTexture2DRHIRef & InDstTexture,
 			int32 YHeight = CVPixelBufferGetHeightOfPlane(InImageBufferRef, 0);
 
 			CVMetalTextureRef YTextureRef = nullptr;
-			CVReturn Result = CVMetalTextureCacheCreateTextureFromImage(kCFAllocatorDefault, MetalTextureCache, InImageBufferRef, nullptr, MTLPixelFormatR8Unorm, YWidth, YHeight, 0, &YTextureRef);
+			CVReturn Result = CVMetalTextureCacheCreateTextureFromImage(kCFAllocatorDefault, MetalTextureCache, InImageBufferRef, nullptr, bIs8Bit ? MTLPixelFormatR8Unorm : MTLPixelFormatR16Unorm, YWidth, YHeight, 0, &YTextureRef);
 			check(Result == kCVReturnSuccess);
 			check(YTextureRef);
 
@@ -204,19 +198,19 @@ void FElectraMediaTexConvApple::ConvertTexture(FTexture2DRHIRef & InDstTexture,
 			int32 UVHeight = CVPixelBufferGetHeightOfPlane(InImageBufferRef, 1);
 
 			CVMetalTextureRef UVTextureRef = nullptr;
-			Result = CVMetalTextureCacheCreateTextureFromImage(kCFAllocatorDefault, MetalTextureCache, InImageBufferRef, nullptr, MTLPixelFormatRG8Unorm, UVWidth, UVHeight, 1, &UVTextureRef);
+			Result = CVMetalTextureCacheCreateTextureFromImage(kCFAllocatorDefault, MetalTextureCache, InImageBufferRef, nullptr, bIs8Bit ? MTLPixelFormatRG8Unorm : MTLPixelFormatRG16Unorm, UVWidth, UVHeight, 1, &UVTextureRef);
 			check(Result == kCVReturnSuccess);
 			check(UVTextureRef);
 
 			// Metal can upload directly from an IOSurface to a 2D texture, so we can just wrap it.
 
 			const FRHITextureCreateDesc YDesc =
-				FRHITextureCreateDesc::Create2D(TEXT("YTex"), YWidth, YHeight, PF_G8)
+				FRHITextureCreateDesc::Create2D(TEXT("YTex"), YWidth, YHeight, bIs8Bit ? PF_G8 : PF_G16)
 				.SetFlags(ETextureCreateFlags::Dynamic | ETextureCreateFlags::NoTiling | ETextureCreateFlags::ShaderResource)
 				.SetBulkData(new FTexConvTexResourceWrapper(YTextureRef));
 
 			const FRHITextureCreateDesc UVDesc =
-				FRHITextureCreateDesc::Create2D(TEXT("UVTex"), UVWidth, UVHeight, PF_R8G8)
+				FRHITextureCreateDesc::Create2D(TEXT("UVTex"), UVWidth, UVHeight, bIs8Bit ? PF_R8G8 : PF_G16R16)
 				.SetFlags(ETextureCreateFlags::Dynamic | ETextureCreateFlags::NoTiling | ETextureCreateFlags::ShaderResource)
 				.SetBulkData(new FTexConvTexResourceWrapper(UVTextureRef));
 
@@ -227,9 +221,7 @@ void FElectraMediaTexConvApple::ConvertTexture(FTexture2DRHIRef & InDstTexture,
 			FRHICommandListImmediate& RHICmdList = FRHICommandListExecutor::GetImmediateCommandList();
 			{
 				// configure media shaders
-				auto ShaderMap = GetGlobalShaderMap(GMaxRHIFeatureLevel);
-				TShaderMapRef<FMediaShadersVS> VertexShader(ShaderMap);
-				TShaderMapRef<FYCbCrConvertPS> PixelShader(ShaderMap);
+				auto GlobalShaderMap = GetGlobalShaderMap(GMaxRHIFeatureLevel);
 
 				FRHIRenderPassInfo RPInfo(InDstTexture, ERenderTargetActions::Load_Store);
 				RHICmdList.BeginRenderPass(RPInfo, TEXT("AvfMediaSampler"));
@@ -240,15 +232,51 @@ void FElectraMediaTexConvApple::ConvertTexture(FTexture2DRHIRef & InDstTexture,
 					GraphicsPSOInit.BlendState = TStaticBlendStateWriteMask<CW_RGBA, CW_NONE, CW_NONE, CW_NONE, CW_NONE, CW_NONE, CW_NONE, CW_NONE>::GetRHI();
 					GraphicsPSOInit.RasterizerState = TStaticRasterizerState<>::GetRHI();
 					GraphicsPSOInit.DepthStencilState = TStaticDepthStencilState<false, CF_Always>::GetRHI();
-
 					GraphicsPSOInit.BoundShaderState.VertexDeclarationRHI = GMediaVertexDeclaration.VertexDeclarationRHI;
-					GraphicsPSOInit.BoundShaderState.VertexShaderRHI = VertexShader.GetVertexShader();
-					GraphicsPSOInit.BoundShaderState.PixelShaderRHI = PixelShader.GetPixelShader();
 					GraphicsPSOInit.PrimitiveType = PT_TriangleStrip;
 
-					SetGraphicsPipelineState(RHICmdList, GraphicsPSOInit, 0);
+					if (Format == EMediaTextureSampleFormat::CharNV12)
+					{
+						//
+						// NV12
+						// (this disrespects any "working color space")
+						//
+
+						TShaderMapRef<FMediaShadersVS> VertexShader(GlobalShaderMap);
+						TShaderMapRef<FYCbCrConvertPS> PixelShader(GlobalShaderMap);
+
+						PixelShader->SetParameters(RHICmdList, YTex, UVTex, *ColorTransform, Off, true);
+
+						GraphicsPSOInit.BoundShaderState.VertexShaderRHI = VertexShader.GetVertexShader();
+						GraphicsPSOInit.BoundShaderState.PixelShaderRHI = PixelShader.GetPixelShader();
+					}
+					else
+					{
+						//
+						// P010
+						//
 
-					PixelShader->SetParameters(RHICmdList, YTex, UVTex, *ColorTransform, MediaShaders::YUVOffset8bits, true);
+						// Setup conversion from Rec2020 to current working color space
+						const UE::Color::FColorSpace& Working = UE::Color::FColorSpace::GetWorking();
+						FMatrix44f ColorSpaceMtx = UE::Color::Transpose<float>(Working.GetXYZToRgb()) * GamutToXYZMtx;
+						// Normalize output (e.g. 80 or 100 nits == 1.0)
+						ColorSpaceMtx = ColorSpaceMtx.ApplyScale(kMediaSample_HDR_NitsNormalizationFactor);
+
+						// Get shaders.
+						TShaderMapRef<FP010ConvertPS> PixelShader(GlobalShaderMap);
+						TShaderMapRef<FMediaShadersVS> VertexShader(GlobalShaderMap);
+
+						FShaderResourceViewRHIRef Y_SRV = RHICreateShaderResourceView(YTex, 0, 1, PF_G16);
+						FShaderResourceViewRHIRef UV_SRV = RHICreateShaderResourceView(UVTex, 0, 1, PF_G16R16);
+
+						// Update shader uniform parameters.
+						PixelShader->SetParameters(RHICmdList, { YWidth, YHeight }, Y_SRV, UV_SRV, { YWidth, YHeight }, YUVMtx, ColorSpaceMtx, EncodingType == UE::Color::EEncoding::ST2084);
+
+						GraphicsPSOInit.BoundShaderState.VertexShaderRHI = VertexShader.GetVertexShader();
+						GraphicsPSOInit.BoundShaderState.PixelShaderRHI = PixelShader.GetPixelShader();
+					}
+
+					SetGraphicsPipelineState(RHICmdList, GraphicsPSOInit, 0);
 
 					FBufferRHIRef VertexBuffer = CreateTempMediaVertexBuffer();
 					RHICmdList.SetStreamSource(0, VertexBuffer, 0);
@@ -264,6 +292,10 @@ void FElectraMediaTexConvApple::ConvertTexture(FTexture2DRHIRef & InDstTexture,
 		}
 		else
 		{
+			//
+			// sRGB
+			//
+
 			//
 			// Grab data directly from image buffer reference
 			// (this will create the output texture here - but will always use the settings here, as data is not converted)
@@ -293,6 +325,60 @@ void FElectraMediaTexConvApple::ConvertTexture(FTexture2DRHIRef & InDstTexture,
 // ------------------------------------------------------------------------------------------------------
 // ------------------------------------------------------------------------------------------------------
 
+void FElectraTextureSample::Initialize(FVideoDecoderOutput* InVideoDecoderOutput)
+{
+	VideoDecoderOutput = StaticCastSharedPtr<FVideoDecoderOutputApple, IDecoderOutputPoolable, ESPMode::ThreadSafe>(InVideoDecoderOutput->AsShared());
+
+	HDRInfo = VideoDecoderOutput->GetHDRInformation();
+	Colorimetry = VideoDecoderOutput->GetColorimetry();
+
+	bool bFullRange = false;
+	if (auto PinnedColorimetry = Colorimetry.Pin())
+	{
+		bFullRange = (PinnedColorimetry->GetMPEGDefinition()->VideoFullRangeFlag != 0);
+	}
+
+	// Prepare YUV -> RGB matrix containing all necessary offsets and scales to produce RGB straight from sample data
+	const FMatrix* Mtx = bFullRange ? &MediaShaders::YuvToRgbRec709Unscaled : &MediaShaders::YuvToRgbRec709Scaled;
+	FVector Off = (GetFormat() == EMediaTextureSampleFormat::P010) ? (bFullRange ? MediaShaders::YUVOffsetNoScale16bits : MediaShaders::YUVOffset16bits)
+		: (bFullRange ? MediaShaders::YUVOffsetNoScale8bits : MediaShaders::YUVOffset8bits);
+	float Scale = 1.0f;
+
+	if (auto PinnedHDRInfo = HDRInfo.Pin())
+	{
+		switch (PinnedHDRInfo->GetHDRType())
+		{
+		case IVideoDecoderHDRInformation::EType::PQ10:
+		case IVideoDecoderHDRInformation::EType::HLG10:
+		case IVideoDecoderHDRInformation::EType::HDR10:
+			Mtx = &MediaShaders::YuvToRgbRec2020Unscaled;
+			Off = MediaShaders::YUVOffsetNoScale16bits;
+			break;
+		default:
+			break;
+		}
+	}
+
+	// Matrix to transform sample data to standard YUV values
+	FMatrix PreMtx = FMatrix::Identity;
+	PreMtx.M[0][0] = Scale;
+	PreMtx.M[1][1] = Scale;
+	PreMtx.M[2][2] = Scale;
+	PreMtx.M[0][3] = -Off.X;
+	PreMtx.M[1][3] = -Off.Y;
+	PreMtx.M[2][3] = -Off.Z;
+
+	// Combine this with the actual YUV-RGB conversion
+	YuvToRgbMtx = FMatrix44f(*Mtx * PreMtx);
+}
+
+
+EMediaTextureSampleFormat FElectraTextureSample::GetFormat() const
+{
+	OSType PixelFormat = CVPixelBufferGetPixelFormatType(VideoDecoderOutput->GetImageBuffer());
+	return ((PixelFormat == kCVPixelFormatType_420YpCbCr8BiPlanarVideoRange) || (PixelFormat == kCVPixelFormatType_420YpCbCr8BiPlanarFullRange)) ? EMediaTextureSampleFormat::CharNV12 : EMediaTextureSampleFormat::P010;
+}
+
 
 /**
 * Init code for realloacting an image from the the pool
@@ -372,11 +458,103 @@ bool FElectraTextureSample::Convert(FTexture2DRHIRef & InDstTexture, const FConv
 	{
 		if (TSharedPtr<FElectraMediaTexConvApple, ESPMode::ThreadSafe> PinnedTexConv = TexConv.Pin())
 		{
-			PinnedTexConv->ConvertTexture(InDstTexture, VideoDecoderOutput->GetImageBuffer());
+			PinnedTexConv->ConvertTexture(InDstTexture, VideoDecoderOutput->GetImageBuffer(), VideoDecoderOutput->GetColorimetry()->GetMPEGDefinition()->VideoFullRangeFlag != 0, GetFormat(), GetSampleToRGBMatrix(), GetGamutToXYZMatrix(), GetEncodingType());
 			return true;
 		}
 	}
 	return false;
 }
 
+bool FElectraTextureSample::IsOutputSrgb() const
+{
+	if (auto PinnedHDRInfo = HDRInfo.Pin())
+	{
+		// If the HDR type is unknown we also assume sRGB
+		return (PinnedHDRInfo->GetHDRType() == IVideoDecoderHDRInformation::EType::Unknown);
+	}
+	return true;
+}
+
+const FMatrix& FElectraTextureSample::GetYUVToRGBMatrix() const
+{
+	if (auto PinnedHDRInfo = HDRInfo.Pin())
+	{
+		switch (PinnedHDRInfo->GetHDRType())
+		{
+			case IVideoDecoderHDRInformation::EType::PQ10:
+			case IVideoDecoderHDRInformation::EType::HLG10:
+			case IVideoDecoderHDRInformation::EType::HDR10:		return MediaShaders::YuvToRgbRec2020Unscaled;
+			case IVideoDecoderHDRInformation::EType::Unknown:	break;
+		}
+	}
+
+	// If no HDR info is present, we assume sRGB
+	return GetFullRange() ? MediaShaders::YuvToRgbRec709Unscaled : MediaShaders::YuvToRgbRec709Scaled;
+}
+
+bool FElectraTextureSample::GetFullRange() const
+{
+	bool bFullRange = false;
+	if (auto PinnedColorimetry = Colorimetry.Pin())
+	{
+		bFullRange = (PinnedColorimetry->GetMPEGDefinition()->VideoFullRangeFlag != 0);
+	}
+	return bFullRange;
+}
+
+FMatrix44f FElectraTextureSample::GetSampleToRGBMatrix() const
+{
+	return YuvToRgbMtx;
+}
+
+FMatrix44f FElectraTextureSample::GetGamutToXYZMatrix() const
+{
+	if (auto PinnedHDRInfo = HDRInfo.Pin())
+	{
+		switch (PinnedHDRInfo->GetHDRType())
+		{
+			case IVideoDecoderHDRInformation::EType::PQ10:
+			case IVideoDecoderHDRInformation::EType::HLG10:
+			case IVideoDecoderHDRInformation::EType::HDR10:		return GamutToXYZMatrix(EDisplayColorGamut::Rec2020_D65);
+			case IVideoDecoderHDRInformation::EType::Unknown:	return GamutToXYZMatrix(EDisplayColorGamut::sRGB_D65);
+		}
+	}
+	// If no HDR info is present, we assume sRGB
+	return GamutToXYZMatrix(EDisplayColorGamut::sRGB_D65);
+}
+
+FVector2f FElectraTextureSample::GetWhitePoint() const
+{
+	if (auto PinnedHDRInfo = HDRInfo.Pin())
+	{
+		if (auto ColorVolume = PinnedHDRInfo->GetMasteringDisplayColourVolume())
+		{
+			return FVector2f(ColorVolume->white_point_x, ColorVolume->white_point_y);
+		}
+	}
+	// If no HDR info is present, we assume sRGB
+	return FVector2f(UE::Color::GetWhitePoint(UE::Color::EWhitePoint::CIE1931_D65));
+}
+
+UE::Color::EEncoding FElectraTextureSample::GetEncodingType() const
+{
+	if (auto PinnedHDRInfo = HDRInfo.Pin())
+	{
+		switch (PinnedHDRInfo->GetHDRType())
+		{
+			case IVideoDecoderHDRInformation::EType::PQ10:
+			case IVideoDecoderHDRInformation::EType::HDR10:		return UE::Color::EEncoding::ST2084;
+			case IVideoDecoderHDRInformation::EType::Unknown:	return UE::Color::EEncoding::sRGB;
+			case IVideoDecoderHDRInformation::EType::HLG10:
+			{
+				check(!"*** Implement support for HLG10 in UE::Color!");
+				return UE::Color::EEncoding::sRGB;
+			}
+		}
+	}
+
+	// No info: use our legacy value
+	return UE::Color::EEncoding::sRGB;
+}
+
 #endif
diff --git a/Engine/Plugins/Media/ElectraUtil/Source/ElectraSamples/Public/Apple/ElectraTextureSample.h b/Engine/Plugins/Media/ElectraUtil/Source/ElectraSamples/Public/Apple/ElectraTextureSample.h
--- a/Engine/Plugins/Media/ElectraUtil/Source/ElectraSamples/Public/Apple/ElectraTextureSample.h
+++ b/Engine/Plugins/Media/ElectraUtil/Source/ElectraSamples/Public/Apple/ElectraTextureSample.h
@@ -33,10 +33,7 @@ public:
 	}
 
 public:
-	void Initialize(FVideoDecoderOutput* InVideoDecoderOutput)
-	{
-		VideoDecoderOutput = StaticCastSharedPtr<FVideoDecoderOutputApple, IDecoderOutputPoolable, ESPMode::ThreadSafe>(InVideoDecoderOutput->AsShared());
-	}
+	void Initialize(FVideoDecoderOutput* InVideoDecoderOutput);
 
 	//~ IMediaTextureSample interface
 
@@ -59,10 +56,7 @@ public:
 		return (EMediaOrientation)VideoDecoderOutput->GetOrientation();
 	}
 
-	virtual EMediaTextureSampleFormat GetFormat() const override
-	{
-		return EMediaTextureSampleFormat::CharBGRA;
-	}
+	virtual EMediaTextureSampleFormat GetFormat() const override;
 
 	virtual FIntPoint GetOutputDim() const override;
 	virtual uint32 GetStride() const override;
@@ -79,16 +73,19 @@ public:
 		return true;
 	}
 
-	virtual bool IsOutputSrgb() const override
-	{
-		return true;
-	}
-
     virtual IMediaTextureSampleConverter* GetMediaTextureSampleConverter() override
     {
         return this;
     }
 
+	bool IsOutputSrgb() const override;
+	const FMatrix& GetYUVToRGBMatrix() const override;
+	bool GetFullRange() const override;
+	FMatrix44f GetSampleToRGBMatrix() const override;
+	FMatrix44f GetGamutToXYZMatrix() const override;
+	FVector2f GetWhitePoint() const override;
+	UE::Color::EEncoding GetEncodingType() const override;
+
 	virtual void InitializePoolable() override;
 	virtual void ShutdownPoolable() override;
 
@@ -101,6 +98,13 @@ private:
 
 	TWeakPtr<FElectraMediaTexConvApple, ESPMode::ThreadSafe> TexConv;
 
+	/** Quick access for some HDR related info */
+	TWeakPtr<const IVideoDecoderHDRInformation, ESPMode::ThreadSafe> HDRInfo;
+	TWeakPtr<const IVideoDecoderColorimetry, ESPMode::ThreadSafe> Colorimetry;
+
+	/** YUV matrix, adjusted to compensate for decoder output specific scale */
+	FMatrix44f YuvToRgbMtx;
+
 	virtual uint32 GetConverterInfoFlags() const override;
     virtual bool Convert(FTexture2DRHIRef & InDstTexture, const FConversionHints & Hints) override;
 };
@@ -117,7 +121,7 @@ public:
     ~FElectraMediaTexConvApple();
 
 #if WITH_ENGINE
-    void ConvertTexture(FTexture2DRHIRef & InDstTexture, CVImageBufferRef InImageBufferRef);
+    void ConvertTexture(FTexture2DRHIRef & InDstTexture, CVImageBufferRef InImageBufferRef, bool bFullRange, EMediaTextureSampleFormat Format, const FMatrix44f& YUVMtx, const FMatrix44f& GamutToXYZMtx, UE::Color::EEncoding EncodingType);
 #endif
 
 private:
