Introducing a baseclass for all ELectraTextureSamples & adding AVF & Android as supported platforms

diff --git a/Engine/Plugins/Media/ElectraPlayer/Source/ElectraPlayerRuntime/Private/Runtime/Decoder/Android/VideoDecoderH265_Android.cpp b/Engine/Plugins/Media/ElectraPlayer/Source/ElectraPlayerRuntime/Private/Runtime/Decoder/Android/VideoDecoderH265_Android.cpp
--- a/Engine/Plugins/Media/ElectraPlayer/Source/ElectraPlayerRuntime/Private/Runtime/Decoder/Android/VideoDecoderH265_Android.cpp
+++ b/Engine/Plugins/Media/ElectraPlayer/Source/ElectraPlayerRuntime/Private/Runtime/Decoder/Android/VideoDecoderH265_Android.cpp
@@ -1498,7 +1498,7 @@ FVideoDecoderH265::EOutputResult FVideoDecoderH265::GetOutput()
 
 				// NOTE: We do not release the output buffer here as this would immediately overwrite the last frame in the surface not yet rendered.
 				//	     Instead we leave returning the buffer and the subsequent update of the texture to the render task.
-					//Result = DecoderInstance->ReleaseOutputBuffer(OutputBufferInfo.BufferIndex, OutputBufferInfo.ValidCount, false, -1);
+				//Result = DecoderInstance->ReleaseOutputBuffer(OutputBufferInfo.BufferIndex, OutputBufferInfo.ValidCount, false, -1);
 			}
 			else
 			{
@@ -1587,13 +1587,12 @@ FVideoDecoderH265::EOutputResult FVideoDecoderH265::ProcessOutput(const FDecoded
 					OutputBufferSampleProperties->Set("aspect_h", FVariantValue(ay));
 					OutputBufferSampleProperties->Set("fps_num", FVariantValue((int64)0));
 					OutputBufferSampleProperties->Set("fps_denom", FVariantValue((int64)0));
-					OutputBufferSampleProperties->Set("pixelfmt", FVariantValue((int64)EPixelFormat::PF_B8G8R8A8));
 
 					// Set the bit depth and the colorimetry.
 					uint8 colour_primaries=2, transfer_characteristics=2, matrix_coeffs=2;
 					uint8 video_full_range_flag=0, video_format=5;
 					uint8 num_bits = 8;
-					if (NextImage.SourceInfo->SPSs.Num())
+					if (NextImage.SourceInfo.IsValid() && NextImage.SourceInfo->SPSs.Num())
 					{
 						check(NextImage.SourceInfo->SPSs[0].bit_depth_luma_minus8 == NextImage.SourceInfo->SPSs[0].bit_depth_chroma_minus8);
 						num_bits = NextImage.SourceInfo->SPSs[0].bit_depth_luma_minus8 + 8;
@@ -1609,7 +1608,8 @@ FVideoDecoderH265::EOutputResult FVideoDecoderH265::ProcessOutput(const FDecoded
 							video_format = NextImage.SourceInfo->SPSs[0].video_format;
 						}
 					}
-					/// see above // OutputBufferSampleProperties->Set("pixelfmt", num_bits==8 ? FVariantValue((int64)EPixelFormat::PF_NV12) : FVariantValue((int64)EPixelFormat::PF_PLATFORM_HDR_1));
+					
+					OutputBufferSampleProperties->Set("pixelfmt", FVariantValue((int64)((num_bits > 8)  ? EPixelFormat::PF_A2B10G10R10 : EPixelFormat::PF_B8G8R8A8)));
 					OutputBufferSampleProperties->Set("bits_per", FVariantValue((int64)num_bits));
 					Colorimetry.Update(colour_primaries, transfer_characteristics, matrix_coeffs, video_full_range_flag, video_format);
 					Colorimetry.UpdateParamDict(*OutputBufferSampleProperties);
diff --git a/Engine/Plugins/Media/ElectraPlayer/Source/ElectraPlayerRuntime/Private/Runtime/Decoder/Linux/ElectraVideoDecoder_Linux.cpp b/Engine/Plugins/Media/ElectraPlayer/Source/ElectraPlayerRuntime/Private/Runtime/Decoder/Linux/ElectraVideoDecoder_Linux.cpp
--- a/Engine/Plugins/Media/ElectraPlayer/Source/ElectraPlayerRuntime/Private/Runtime/Decoder/Linux/ElectraVideoDecoder_Linux.cpp
+++ b/Engine/Plugins/Media/ElectraPlayer/Source/ElectraPlayerRuntime/Private/Runtime/Decoder/Linux/ElectraVideoDecoder_Linux.cpp
@@ -24,8 +24,8 @@ bool FElectraPlayerVideoDecoderOutputLinux::InitializeForBuffer(FIntPoint Dim, E
 {
 	FVideoDecoderOutputLinux::Initialize(InParamDict);
 	NumBits = InNumBits;
-	check(PixFmt == EPixelFormat::PF_NV12);
-	if (PixFmt != EPixelFormat::PF_NV12)
+	check(PixFmt == EPixelFormat::PF_NV12 || PixFmt == EPixelFormat::PF_P010);
+	if (PixFmt != EPixelFormat::PF_NV12 && PixFmt != EPixelFormat::PF_P010)
 	{
 		return false;
 	}
@@ -72,11 +72,6 @@ uint32 FElectraPlayerVideoDecoderOutputLinux::GetStride() const
 	return Stride;
 }
 
-bool FElectraPlayerVideoDecoderOutputLinux::GetIs10Bit() const
-{
-	return NumBits > 8;
-}
-
 
 void FElectraPlayerVideoDecoderOutputLinux::SetOwner(const TSharedPtr<IDecoderOutputOwner, ESPMode::ThreadSafe>& InOwningRenderer)
 {
diff --git a/Engine/Plugins/Media/ElectraPlayer/Source/ElectraPlayerRuntime/Private/Runtime/Decoder/Linux/ElectraVideoDecoder_Linux.h b/Engine/Plugins/Media/ElectraPlayer/Source/ElectraPlayerRuntime/Private/Runtime/Decoder/Linux/ElectraVideoDecoder_Linux.h
--- a/Engine/Plugins/Media/ElectraPlayer/Source/ElectraPlayerRuntime/Private/Runtime/Decoder/Linux/ElectraVideoDecoder_Linux.h
+++ b/Engine/Plugins/Media/ElectraPlayer/Source/ElectraPlayerRuntime/Private/Runtime/Decoder/Linux/ElectraVideoDecoder_Linux.h
@@ -30,7 +30,6 @@ public:
 	uint32 GetStride() const override;
 	FIntPoint GetDim() const override
 	{ return SampleDim; }
-	bool GetIs10Bit() const override;
 
 	void SetOwner(const TSharedPtr<IDecoderOutputOwner, ESPMode::ThreadSafe>& InOwningRenderer) override;
 
diff --git a/Engine/Plugins/Media/ElectraPlayer/Source/ElectraPlayerRuntime/Private/Runtime/Decoder/Linux/VideoDecoderH265_Linux.cpp b/Engine/Plugins/Media/ElectraPlayer/Source/ElectraPlayerRuntime/Private/Runtime/Decoder/Linux/VideoDecoderH265_Linux.cpp
--- a/Engine/Plugins/Media/ElectraPlayer/Source/ElectraPlayerRuntime/Private/Runtime/Decoder/Linux/VideoDecoderH265_Linux.cpp
+++ b/Engine/Plugins/Media/ElectraPlayer/Source/ElectraPlayerRuntime/Private/Runtime/Decoder/Linux/VideoDecoderH265_Linux.cpp
@@ -1399,7 +1399,7 @@ void FVideoDecoderH265LinuxLibavcodec::ProcessOutput(bool bFlush)
 				OutputBufferSampleProperties->Set("aspect_h", FVariantValue((int64) ay));
 				OutputBufferSampleProperties->Set("fps_num", FVariantValue((int64) 0 ));
 				OutputBufferSampleProperties->Set("fps_denom", FVariantValue((int64) 0 ));
-				OutputBufferSampleProperties->Set("pixelfmt", FVariantValue((int64)EPixelFormat::PF_NV12));
+				OutputBufferSampleProperties->Set("pixelfmt", FVariantValue((int64)((num_bits <= 8) ? EPixelFormat::PF_NV12 : EPixelFormat::PF_P010)));
 
 				// Set the bit depth and the colorimetry.
 				uint8 colour_primaries=2, transfer_characteristics=2, matrix_coeffs=2;
@@ -1420,7 +1420,7 @@ void FVideoDecoderH265LinuxLibavcodec::ProcessOutput(bool bFlush)
 						video_format = NextImage.SourceInfo->SPSs[0].video_format;
 					}
 				}
-			/// see above // OutputBufferSampleProperties->Set("pixelfmt", num_bits==8 ? FVariantValue((int64)EPixelFormat::PF_NV12) : FVariantValue((int64)EPixelFormat::PF_PLATFORM_HDR_1));
+			/// see above // OutputBufferSampleProperties->Set("pixelfmt", num_bits==8 ? FVariantValue((int64)EPixelFormat::PF_NV12) : FVariantValue((int64)EPixelFormat::PF_P010));
 				OutputBufferSampleProperties->Set("bits_per", FVariantValue((int64)num_bits));
 				Colorimetry.Update(colour_primaries, transfer_characteristics, matrix_coeffs, video_full_range_flag, video_format);
 				Colorimetry.UpdateParamDict(*OutputBufferSampleProperties);
diff --git a/Engine/Plugins/Media/ElectraUtil/Source/ElectraBase/Public/Linux/MediaVideoDecoderOutputLinux.h b/Engine/Plugins/Media/ElectraUtil/Source/ElectraBase/Public/Linux/MediaVideoDecoderOutputLinux.h
--- a/Engine/Plugins/Media/ElectraUtil/Source/ElectraBase/Public/Linux/MediaVideoDecoderOutputLinux.h
+++ b/Engine/Plugins/Media/ElectraUtil/Source/ElectraBase/Public/Linux/MediaVideoDecoderOutputLinux.h
@@ -12,5 +12,4 @@ public:
 	virtual const TArray<uint8>& GetBuffer() const = 0;
 	virtual FIntPoint GetBufferDimensions() const = 0;
 	virtual uint32 GetStride() const = 0;
-	virtual bool GetIs10Bit() const = 0;
 };
diff --git a/Engine/Plugins/Media/ElectraUtil/Source/ElectraSamples/Private/ElectraTextureSampleAVF.cpp b/Engine/Plugins/Media/ElectraUtil/Source/ElectraSamples/Private/ElectraTextureSampleAVF.cpp
--- a/Engine/Plugins/Media/ElectraUtil/Source/ElectraSamples/Private/ElectraTextureSampleAVF.cpp
+++ b/Engine/Plugins/Media/ElectraUtil/Source/ElectraSamples/Private/ElectraTextureSampleAVF.cpp
@@ -330,118 +330,22 @@ void FElectraMediaTexConvApple::ConvertTexture(FTexture2DRHIRef & InDstTexture,
 
 void FElectraTextureSample::Initialize(FVideoDecoderOutput* InVideoDecoderOutput)
 {
-	VideoDecoderOutput = StaticCastSharedPtr<FVideoDecoderOutputApple, IDecoderOutputPoolable, ESPMode::ThreadSafe>(InVideoDecoderOutput->AsShared());
-
-	HDRInfo = VideoDecoderOutput->GetHDRInformation();
-	Colorimetry = VideoDecoderOutput->GetColorimetry();
-
-	bool bFullRange = false;
-	if (auto PinnedColorimetry = Colorimetry.Pin())
-	{
-		bFullRange = (PinnedColorimetry->GetMPEGDefinition()->VideoFullRangeFlag != 0);
-	}
-
-	// Prepare YUV -> RGB matrix containing all necessary offsets and scales to produce RGB straight from sample data
-	const FMatrix* Mtx = bFullRange ? &MediaShaders::YuvToRgbRec709Unscaled : &MediaShaders::YuvToRgbRec709Scaled;
-	FVector Off = (GetFormat() == EMediaTextureSampleFormat::P010) ? (bFullRange ? MediaShaders::YUVOffsetNoScale16bits : MediaShaders::YUVOffset16bits)
-		: (bFullRange ? MediaShaders::YUVOffsetNoScale8bits : MediaShaders::YUVOffset8bits);
-	float Scale = 1.0f;
-
-	if (auto PinnedHDRInfo = HDRInfo.Pin())
-	{
-		switch (PinnedHDRInfo->GetHDRType())
-		{
-		case IVideoDecoderHDRInformation::EType::PQ10:
-		case IVideoDecoderHDRInformation::EType::HLG10:
-		case IVideoDecoderHDRInformation::EType::HDR10:
-			Mtx = &MediaShaders::YuvToRgbRec2020Unscaled;
-			Off = MediaShaders::YUVOffsetNoScale16bits;
-			break;
-		default:
-			break;
-		}
-	}
-
-	// Matrix to transform sample data to standard YUV values
-	FMatrix PreMtx = FMatrix::Identity;
-	PreMtx.M[0][0] = Scale;
-	PreMtx.M[1][1] = Scale;
-	PreMtx.M[2][2] = Scale;
-	PreMtx.M[0][3] = -Off.X;
-	PreMtx.M[1][3] = -Off.Y;
-	PreMtx.M[2][3] = -Off.Z;
-
-	// Combine this with the actual YUV-RGB conversion
-	YuvToRgbMtx = FMatrix44f(*Mtx * PreMtx);
+	IElectraTextureSampleBase::Initialize(InVideoDecoderOutput);
+	VideoDecoderOutputApple = static_cast<FVideoDecoderOutputApple*>(InVideoDecoderOutput);
 }
 
 
 EMediaTextureSampleFormat FElectraTextureSample::GetFormat() const
 {
-	OSType PixelFormat = CVPixelBufferGetPixelFormatType(VideoDecoderOutput->GetImageBuffer());
+	OSType PixelFormat = CVPixelBufferGetPixelFormatType(VideoDecoderOutputApple->GetImageBuffer());
 	return ((PixelFormat == kCVPixelFormatType_420YpCbCr8BiPlanarVideoRange) || (PixelFormat == kCVPixelFormatType_420YpCbCr8BiPlanarFullRange)) ? EMediaTextureSampleFormat::CharNV12 : EMediaTextureSampleFormat::P010;
 }
 
-
-/**
-* Init code for realloacting an image from the the pool
-*/
-void FElectraTextureSample::InitializePoolable()
-{
-}
-
-
-/**
-*  Return the object to the pool and inform the renderer about this...
-*/
-void FElectraTextureSample::ShutdownPoolable()
-{
-	VideoDecoderOutput.Reset();
-}
-
-
-FIntPoint FElectraTextureSample::GetDim() const
-{
-	if (VideoDecoderOutput)
-	{
-		return VideoDecoderOutput->GetDim();
-	}
-	return FIntPoint::ZeroValue;
-}
-
-FMediaTimeStamp FElectraTextureSample::GetTime() const
-{
-	if (VideoDecoderOutput)
-	{
-		const FDecoderTimeStamp TimeStamp = VideoDecoderOutput->GetTime();
-		return FMediaTimeStamp(TimeStamp.Time, TimeStamp.SequenceIndex);
-	}
-	return FMediaTimeStamp();
-}
-
-FTimespan FElectraTextureSample::GetDuration() const
-{
-	if (VideoDecoderOutput)
-	{
-		return VideoDecoderOutput->GetDuration();
-	}
-	return FTimespan::Zero();
-}
-
-FIntPoint FElectraTextureSample::GetOutputDim() const
-{
-	if (VideoDecoderOutput)
-	{
-		return VideoDecoderOutput->GetOutputDim();
-	}
-	return FIntPoint::ZeroValue;
-}
-
 uint32 FElectraTextureSample::GetStride() const
 {
 	if (VideoDecoderOutput)
 	{
-		return VideoDecoderOutput->GetStride();
+		return VideoDecoderOutputApple->GetStride();
 	}
 	return 0;
 }
@@ -450,7 +354,7 @@ uint32 FElectraTextureSample::GetConverterInfoFlags() const
 {
 	if (VideoDecoderOutput)
 	{
-		return CVPixelBufferIsPlanar(VideoDecoderOutput->GetImageBuffer()) ? ConverterInfoFlags_Default : ConverterInfoFlags_WillCreateOutputTexture;
+		return CVPixelBufferIsPlanar(VideoDecoderOutputApple->GetImageBuffer()) ? ConverterInfoFlags_Default : ConverterInfoFlags_WillCreateOutputTexture;
 	}
 	return ConverterInfoFlags_Default;
 }
@@ -461,103 +365,11 @@ bool FElectraTextureSample::Convert(FTexture2DRHIRef & InDstTexture, const FConv
 	{
 		if (TSharedPtr<FElectraMediaTexConvApple, ESPMode::ThreadSafe> PinnedTexConv = TexConv.Pin())
 		{
-			PinnedTexConv->ConvertTexture(InDstTexture, VideoDecoderOutput->GetImageBuffer(), VideoDecoderOutput->GetColorimetry()->GetMPEGDefinition()->VideoFullRangeFlag != 0, GetFormat(), GetSampleToRGBMatrix(), GetGamutToXYZMatrix(), GetEncodingType());
+			PinnedTexConv->ConvertTexture(InDstTexture, VideoDecoderOutputApple->GetImageBuffer(), VideoDecoderOutput->GetColorimetry()->GetMPEGDefinition()->VideoFullRangeFlag != 0, GetFormat(), GetSampleToRGBMatrix(), GetGamutToXYZMatrix(), GetEncodingType());
 			return true;
 		}
 	}
 	return false;
 }
 
-bool FElectraTextureSample::IsOutputSrgb() const
-{
-	if (auto PinnedHDRInfo = HDRInfo.Pin())
-	{
-		// If the HDR type is unknown we also assume sRGB
-		return (PinnedHDRInfo->GetHDRType() == IVideoDecoderHDRInformation::EType::Unknown);
-	}
-	return true;
-}
-
-const FMatrix& FElectraTextureSample::GetYUVToRGBMatrix() const
-{
-	if (auto PinnedHDRInfo = HDRInfo.Pin())
-	{
-		switch (PinnedHDRInfo->GetHDRType())
-		{
-			case IVideoDecoderHDRInformation::EType::PQ10:
-			case IVideoDecoderHDRInformation::EType::HLG10:
-			case IVideoDecoderHDRInformation::EType::HDR10:		return MediaShaders::YuvToRgbRec2020Unscaled;
-			case IVideoDecoderHDRInformation::EType::Unknown:	break;
-		}
-	}
-
-	// If no HDR info is present, we assume sRGB
-	return GetFullRange() ? MediaShaders::YuvToRgbRec709Unscaled : MediaShaders::YuvToRgbRec709Scaled;
-}
-
-bool FElectraTextureSample::GetFullRange() const
-{
-	bool bFullRange = false;
-	if (auto PinnedColorimetry = Colorimetry.Pin())
-	{
-		bFullRange = (PinnedColorimetry->GetMPEGDefinition()->VideoFullRangeFlag != 0);
-	}
-	return bFullRange;
-}
-
-FMatrix44f FElectraTextureSample::GetSampleToRGBMatrix() const
-{
-	return YuvToRgbMtx;
-}
-
-FMatrix44f FElectraTextureSample::GetGamutToXYZMatrix() const
-{
-	if (auto PinnedHDRInfo = HDRInfo.Pin())
-	{
-		switch (PinnedHDRInfo->GetHDRType())
-		{
-			case IVideoDecoderHDRInformation::EType::PQ10:
-			case IVideoDecoderHDRInformation::EType::HLG10:
-			case IVideoDecoderHDRInformation::EType::HDR10:		return GamutToXYZMatrix(EDisplayColorGamut::Rec2020_D65);
-			case IVideoDecoderHDRInformation::EType::Unknown:	return GamutToXYZMatrix(EDisplayColorGamut::sRGB_D65);
-		}
-	}
-	// If no HDR info is present, we assume sRGB
-	return GamutToXYZMatrix(EDisplayColorGamut::sRGB_D65);
-}
-
-FVector2f FElectraTextureSample::GetWhitePoint() const
-{
-	if (auto PinnedHDRInfo = HDRInfo.Pin())
-	{
-		if (auto ColorVolume = PinnedHDRInfo->GetMasteringDisplayColourVolume())
-		{
-			return FVector2f(ColorVolume->white_point_x, ColorVolume->white_point_y);
-		}
-	}
-	// If no HDR info is present, we assume sRGB
-	return FVector2f(UE::Color::GetWhitePoint(UE::Color::EWhitePoint::CIE1931_D65));
-}
-
-UE::Color::EEncoding FElectraTextureSample::GetEncodingType() const
-{
-	if (auto PinnedHDRInfo = HDRInfo.Pin())
-	{
-		switch (PinnedHDRInfo->GetHDRType())
-		{
-			case IVideoDecoderHDRInformation::EType::PQ10:
-			case IVideoDecoderHDRInformation::EType::HDR10:		return UE::Color::EEncoding::ST2084;
-			case IVideoDecoderHDRInformation::EType::Unknown:	return UE::Color::EEncoding::sRGB;
-			case IVideoDecoderHDRInformation::EType::HLG10:
-			{
-				check(!"*** Implement support for HLG10 in UE::Color!");
-				return UE::Color::EEncoding::sRGB;
-			}
-		}
-	}
-
-	// No info: use our legacy value
-	return UE::Color::EEncoding::sRGB;
-}
-
 #endif
diff --git a/Engine/Plugins/Media/ElectraUtil/Source/ElectraSamples/Private/ElectraTextureSampleAndroid.cpp b/Engine/Plugins/Media/ElectraUtil/Source/ElectraSamples/Private/ElectraTextureSampleAndroid.cpp
--- a/Engine/Plugins/Media/ElectraUtil/Source/ElectraSamples/Private/ElectraTextureSampleAndroid.cpp
+++ b/Engine/Plugins/Media/ElectraUtil/Source/ElectraSamples/Private/ElectraTextureSampleAndroid.cpp
@@ -12,6 +12,7 @@
 
 #include "RenderingThread.h"
 
+DECLARE_GPU_STAT_NAMED(MediaAndroidDecoder_Convert, TEXT("MediaAndroidDecoder_Convert"));
 
 #define ELECTRA_INIT_ON_RENDERTHREAD	1	// set to 1 if context & surface init should be on render thread (seems safer for compatibility - TODO: research why)
 
@@ -122,7 +123,7 @@ FElectraTextureSampleSupport::FElectraTextureSampleSupport()
 	, InitializeFN(GetClassMethod("Initialize", "(Z)V"))
 	, ReleaseFN(GetClassMethod("Release", "()V"))
 	, GetCodecSurfaceFN(GetClassMethod("GetCodecSurface", "()Landroid/view/Surface;"))
-	, GetVideoFrameUpdateInfoFN(GetClassMethod("GetVideoFrameUpdateInfo", "(III)Lcom/epicgames/unreal/ElectraTextureSample$FFrameUpdateInfo;"))
+	, GetVideoFrameUpdateInfoFN(GetClassMethod("GetVideoFrameUpdateInfo", "(IIIZ)Lcom/epicgames/unreal/ElectraTextureSample$FFrameUpdateInfo;"))
 	, CodecSurface(nullptr)
 	, SurfaceInitEvent(nullptr)
 {
@@ -233,7 +234,7 @@ int32 FElectraTextureSampleSupport::GetFrameDataAndUpdateInfo(FFrameUpdateInfo&
 
 	// Update frame info and get data...
 	JNIEnv* JEnv = FAndroidApplication::GetJavaEnv();
-	jobject OutputInfo = JEnv->CallObjectMethod(Object, GetVideoFrameUpdateInfoFN.Method, DestTexture, InTargetSample->GetDim().X, InTargetSample->GetDim().Y);
+	jobject OutputInfo = JEnv->CallObjectMethod(Object, GetVideoFrameUpdateInfoFN.Method, DestTexture, InTargetSample->GetDim().X, InTargetSample->GetDim().Y, InTargetSample->GetFormat() == EMediaTextureSampleFormat::CharBGR10A2);
 	if (JEnv->ExceptionCheck())
 	{
 		JEnv->ExceptionDescribe();
@@ -286,7 +287,7 @@ int32 FElectraTextureSampleSupport::GetFrameData(FElectraTextureSample* InTarget
 
 	// Update frame info and get data...
 	JNIEnv* JEnv = FAndroidApplication::GetJavaEnv();
-	jobject OutputInfo = JEnv->CallObjectMethod(Object, GetVideoFrameUpdateInfoFN.Method, DestTexture, InTargetSample->GetDim().X, InTargetSample->GetDim().Y);
+	jobject OutputInfo = JEnv->CallObjectMethod(Object, GetVideoFrameUpdateInfoFN.Method, DestTexture, InTargetSample->GetDim().X, InTargetSample->GetDim().Y, InTargetSample->GetFormat() == EMediaTextureSampleFormat::CharBGR10A2);
 	if (JEnv->ExceptionCheck())
 	{
 		JEnv->ExceptionDescribe();
@@ -356,32 +357,17 @@ void* FElectraTextureSamplePool::GetCodecSurface()
 //-----------------------------------------------------------------------------
 //-----------------------------------------------------------------------------
 
-/**
-* Init code for reallocating an image from the the pool
-*/
-void FElectraTextureSample::InitializePoolable()
-{
-}
-
-/**
-*  Return the object to the pool and inform the renderer about this...
-*/
-void FElectraTextureSample::ShutdownPoolable()
-{
-	VideoDecoderOutput.Reset();
-}
-
-
 void FElectraTextureSample::Initialize(FVideoDecoderOutput* InVideoDecoderOutput)
 {
-	VideoDecoderOutput = StaticCastSharedPtr<FVideoDecoderOutputAndroid, IDecoderOutputPoolable, ESPMode::ThreadSafe>(InVideoDecoderOutput->AsShared());
+	IElectraTextureSampleBase::Initialize(InVideoDecoderOutput);
+	VideoDecoderOutputAndroid = static_cast<FVideoDecoderOutputAndroid*>(InVideoDecoderOutput);
 
-	if (VideoDecoderOutput->GetOutputType() == FVideoDecoderOutputAndroid::EOutputType::DirectToSurfaceAsQueue)
+	if (VideoDecoderOutputAndroid->GetOutputType() == FVideoDecoderOutputAndroid::EOutputType::DirectToSurfaceAsQueue)
 	{
-		ENQUEUE_RENDER_COMMAND(InitTextureSample)([WeakThis{ TWeakPtr<FElectraTextureSample, ESPMode::ThreadSafe>(AsShared()) }](FRHICommandListImmediate& RHICmdList) {
+		ENQUEUE_RENDER_COMMAND(InitTextureSample)([WeakThis{ AsWeak() }](FRHICommandListImmediate& RHICmdList) {
 			if (TSharedPtr<FElectraTextureSample, ESPMode::ThreadSafe> This = WeakThis.Pin())
 			{
-				This->InitializeTexture();
+				This->InitializeTexture(This->VideoDecoderOutputAndroid->GetFormat());
 
 				if (This->Texture)
 				{
@@ -405,12 +391,13 @@ void FElectraTextureSample::Initialize(FVideoDecoderOutput* InVideoDecoderOutput
 }
 
 
-void FElectraTextureSample::InitializeTexture()
+void FElectraTextureSample::InitializeTexture(EPixelFormat PixelFormat)
 {
 	check(IsInRenderingThread() || IsInRHIThread());
 
 	if (FAndroidMisc::ShouldUseVulkan())
 	{
+		// For Vulkan we use a CPU-side buffer to transport the data
 		Texture = nullptr;
 		return;
 	}
@@ -419,16 +406,17 @@ void FElectraTextureSample::InitializeTexture()
 
 	if (Texture.IsValid() && (Texture->GetSizeXY() == Dim))
 	{
+		// The existing texture is just fine...
 		return;
 	}
 
+	// Make linear texture of appropriate bit depth to carry data...
 	const FRHITextureCreateDesc Desc =
 		FRHITextureCreateDesc::Create2D(TEXT("FElectraTextureSample"))
 		.SetExtent(Dim)
-		.SetFormat(PF_B8G8R8A8)
-		.SetFlags(ETextureCreateFlags::Dynamic | ETextureCreateFlags::SRGB | ETextureCreateFlags::RenderTargetable | ETextureCreateFlags::ShaderResource)
-		.SetInitialState(ERHIAccess::SRVMask);
-
+		.SetInitialState(ERHIAccess::SRVMask)
+		.SetFlags(ETextureCreateFlags::Dynamic | ETextureCreateFlags::RenderTargetable | ETextureCreateFlags::ShaderResource)
+		.SetFormat(PixelFormat);
 	Texture = RHICreateTexture(Desc);
 
 	return;
@@ -452,4 +440,111 @@ void FElectraTextureSample::SetupFromBuffer(const void* InBuffer, int32 InBuffer
 	FMemory::Memcpy(Buffer, InBuffer, InBufferSize);
 }
 
+
+EMediaTextureSampleFormat FElectraTextureSample::GetFormat() const
+{
+	if (VideoDecoderOutput)
+	{
+		return (VideoDecoderOutput->GetFormat() == EPixelFormat::PF_A2B10G10R10) ? EMediaTextureSampleFormat::CharBGR10A2 : EMediaTextureSampleFormat::CharBGRA;
+	}
+	return EMediaTextureSampleFormat::Undefined;
+}
+
+
+uint32 FElectraTextureSample::GetStride() const
+{
+	// note: we expect RGBA8 or RGB10A2 -> it's always 32 bits
+	return GetDim().X * sizeof(uint32);
+}
+
+
+bool FElectraTextureSample::Convert(FTexture2DRHIRef& InDstTexture, const FConversionHints& Hints)
+{
+	check(IsInRenderingThread());
+
+	if (GDynamicRHI->RHIIsRenderingSuspended())
+	{
+		return false;
+	}
+
+	TRefCountPtr<FRHITexture2D> InputTexture;
+
+	// Either use a texture we have around as a payload or make a temporary one from buffer contents...
+	if (!Texture.IsValid())
+	{
+		auto SampleDim = GetDim();
+
+		// Make a source texture so we can convert from it...
+		const FRHITextureCreateDesc Desc =
+			FRHITextureCreateDesc::Create2D(TEXT("FMediaTextureResource"), SampleDim, VideoDecoderOutputAndroid->GetFormat())
+			.SetFlags(TexCreate_Dynamic);
+		InputTexture = RHICreateTexture(Desc);
+		if (!InputTexture.IsValid())
+		{
+			return false;
+		}
+
+		// copy sample data to input render target
+		FUpdateTextureRegion2D Region(0, 0, 0, 0, SampleDim.X, SampleDim.Y);
+		RHIUpdateTexture2D(InputTexture, 0, Region, GetStride(), (const uint8*)GetBuffer());
+	}
+	else
+	{
+		InputTexture = Texture;
+	}
+
+	FRHICommandListImmediate& RHICmdList = FRHICommandListExecutor::GetImmediateCommandList();
+
+	SCOPED_DRAW_EVENT(RHICmdList, AndroidMediaOutputConvertTexture);
+	SCOPED_GPU_STAT(RHICmdList, MediaAndroidDecoder_Convert);
+
+	FIntPoint Dim = VideoDecoderOutput->GetDim();
+	FIntPoint OutputDim = VideoDecoderOutput->GetOutputDim();
+
+	RHICmdList.Transition(FRHITransitionInfo(InDstTexture, ERHIAccess::Unknown, ERHIAccess::RTV));
+	FRHIRenderPassInfo RPInfo(InDstTexture, ERenderTargetActions::DontLoad_Store);
+	RHICmdList.BeginRenderPass(RPInfo, TEXT("AndroidProcessVideo"));
+
+	// Update viewport.
+	RHICmdList.SetViewport(0, 0, 0.f, OutputDim.X, OutputDim.Y, 1.f);
+
+	// Setup conversion from Rec2020 to current working color space
+	const UE::Color::FColorSpace& Working = UE::Color::FColorSpace::GetWorking();
+	FMatrix44f ColorSpaceMtx = UE::Color::Transpose<float>(Working.GetXYZToRgb()) * GetGamutToXYZMatrix();
+	if (GetEncodingType() == UE::Color::EEncoding::ST2084)
+	{
+		// Normalize output (e.g. 80 or 100 nits == 1.0)
+		ColorSpaceMtx = ColorSpaceMtx.ApplyScale(kMediaSample_HDR_NitsNormalizationFactor);
+	}
+
+	// Get shaders.
+	FGlobalShaderMap* GlobalShaderMap = GetGlobalShaderMap(GMaxRHIFeatureLevel);
+	TShaderMapRef<FRGBConvertPS> PixelShader(GlobalShaderMap);
+	TShaderMapRef<FMediaShadersVS> VertexShader(GlobalShaderMap);
+
+	// Set the graphic pipeline state.
+	FGraphicsPipelineStateInitializer GraphicsPSOInit;
+	RHICmdList.ApplyCachedRenderTargets(GraphicsPSOInit);
+	GraphicsPSOInit.DepthStencilState = TStaticDepthStencilState<false, CF_Always>::GetRHI();
+	GraphicsPSOInit.BlendState = TStaticBlendState<>::GetRHI();
+	GraphicsPSOInit.RasterizerState = TStaticRasterizerState<>::GetRHI();
+	GraphicsPSOInit.PrimitiveType = PT_TriangleStrip;
+	GraphicsPSOInit.BoundShaderState.VertexDeclarationRHI = GMediaVertexDeclaration.VertexDeclarationRHI;
+	GraphicsPSOInit.BoundShaderState.VertexShaderRHI = VertexShader.GetVertexShader();
+	GraphicsPSOInit.BoundShaderState.PixelShaderRHI = PixelShader.GetPixelShader();
+	SetGraphicsPipelineState(RHICmdList, GraphicsPSOInit, 0);
+
+	// Update shader uniform parameters.
+	PixelShader->SetParameters(RHICmdList, InputTexture, OutputDim, GetEncodingType() == UE::Color::EEncoding::sRGB, GetEncodingType() == UE::Color::EEncoding::ST2084, ColorSpaceMtx);
+
+	RHICmdList.SetStreamSource(0, CreateTempMediaVertexBuffer(), 0);
+
+	RHICmdList.DrawPrimitive(0, 2, 1);
+
+	RHICmdList.EndRenderPass();
+	RHICmdList.Transition(FRHITransitionInfo(InDstTexture, ERHIAccess::RTV, ERHIAccess::SRVMask));
+
+	return true;
+}
+
 #endif
diff --git a/Engine/Plugins/Media/ElectraUtil/Source/ElectraSamples/Private/ElectraTextureSampleBase.cpp b/Engine/Plugins/Media/ElectraUtil/Source/ElectraSamples/Private/ElectraTextureSampleBase.cpp
new file mode 100644
--- /dev/null
+++ b/Engine/Plugins/Media/ElectraUtil/Source/ElectraSamples/Private/ElectraTextureSampleBase.cpp
@@ -0,0 +1,209 @@
+// Copyright Epic Games, Inc. All Rights Reserved.
+
+#include "CoreMinimal.h"
+
+#if !UE_SERVER
+
+#include "ElectraTextureSample.h"
+
+// -------------------------------------------------------------------------------------------------------------------------------------------------------
+
+void IElectraTextureSampleBase::Initialize(FVideoDecoderOutput* InVideoDecoderOutput)
+{
+	VideoDecoderOutput = StaticCastSharedPtr<FVideoDecoderOutput, IDecoderOutputPoolable, ESPMode::ThreadSafe>(InVideoDecoderOutput->AsShared());
+	HDRInfo = VideoDecoderOutput->GetHDRInformation();
+	Colorimetry = VideoDecoderOutput->GetColorimetry();
+
+	bool bFullRange = false;
+	if (auto PinnedColorimetry = Colorimetry.Pin())
+	{
+		bFullRange = (PinnedColorimetry->GetMPEGDefinition()->VideoFullRangeFlag != 0);
+	}
+
+	// Prepare YUV -> RGB matrix containing all necessary offsets and scales to produce RGB straight from sample data
+	const FMatrix* Mtx = bFullRange ? &MediaShaders::YuvToRgbRec709Unscaled : &MediaShaders::YuvToRgbRec709Scaled;
+	FVector Off = (VideoDecoderOutput->GetFormat() == PF_NV12) ? (bFullRange ? MediaShaders::YUVOffsetNoScale8bits : MediaShaders::YUVOffset8bits)
+														       : (bFullRange ? MediaShaders::YUVOffsetNoScale16bits : MediaShaders::YUVOffset16bits);
+	float Scale = GetSampleDataScale(false);
+
+	if (auto PinnedHDRInfo = HDRInfo.Pin())
+	{
+		switch (PinnedHDRInfo->GetHDRType())
+		{
+		case IVideoDecoderHDRInformation::EType::PQ10:
+		case IVideoDecoderHDRInformation::EType::HLG10:
+		case IVideoDecoderHDRInformation::EType::HDR10:
+			Mtx = &MediaShaders::YuvToRgbRec2020Unscaled;
+			Off = MediaShaders::YUVOffsetNoScale16bits;
+			Scale = GetSampleDataScale(true);
+			break;
+		default:
+			break;
+		}
+	}
+
+	// Matrix to transform sample data to standard YUV values
+	FMatrix PreMtx = FMatrix::Identity;
+	PreMtx.M[0][0] = Scale;
+	PreMtx.M[1][1] = Scale;
+	PreMtx.M[2][2] = Scale;
+	PreMtx.M[0][3] = -Off.X;
+	PreMtx.M[1][3] = -Off.Y;
+	PreMtx.M[2][3] = -Off.Z;
+
+	// Combine this with the actual YUV-RGB conversion
+	YuvToRgbMtx = FMatrix44f(*Mtx * PreMtx);
+}
+
+void IElectraTextureSampleBase::InitializePoolable()
+{
+}
+
+void IElectraTextureSampleBase::ShutdownPoolable()
+{
+	VideoDecoderOutput.Reset();
+}
+
+FIntPoint IElectraTextureSampleBase::GetDim() const
+{
+	if (VideoDecoderOutput)
+	{
+		return VideoDecoderOutput->GetDim();
+	}
+	return FIntPoint::ZeroValue;
+}
+
+
+FIntPoint IElectraTextureSampleBase::GetOutputDim() const
+{
+	if (VideoDecoderOutput)
+	{
+		return VideoDecoderOutput->GetOutputDim();
+	}
+	return FIntPoint::ZeroValue;
+}
+
+
+FMediaTimeStamp IElectraTextureSampleBase::GetTime() const
+{
+	if (VideoDecoderOutput)
+	{
+		const FDecoderTimeStamp TimeStamp = VideoDecoderOutput->GetTime();
+		return FMediaTimeStamp(TimeStamp.Time, TimeStamp.SequenceIndex);
+	}
+	return FMediaTimeStamp();
+}
+
+
+FTimespan IElectraTextureSampleBase::GetDuration() const
+{
+	if (VideoDecoderOutput)
+	{
+		return VideoDecoderOutput->GetDuration();
+	}
+	return FTimespan::Zero();
+}
+
+bool IElectraTextureSampleBase::IsOutputSrgb() const
+{
+	if (auto PinnedHDRInfo = HDRInfo.Pin())
+	{
+		// If the HDR type is unknown we also assume sRGB
+		return (PinnedHDRInfo->GetHDRType() == IVideoDecoderHDRInformation::EType::Unknown);
+	}
+	// If no HDR info is present, we assume sRGB
+	return true;
+}
+
+const FMatrix& IElectraTextureSampleBase::GetYUVToRGBMatrix() const
+{
+	if (auto PinnedHDRInfo = HDRInfo.Pin())
+	{
+		switch (PinnedHDRInfo->GetHDRType())
+		{
+		case IVideoDecoderHDRInformation::EType::PQ10:
+		case IVideoDecoderHDRInformation::EType::HLG10:
+		case IVideoDecoderHDRInformation::EType::HDR10:		return MediaShaders::YuvToRgbRec2020Unscaled;
+		case IVideoDecoderHDRInformation::EType::Unknown:	break;
+		}
+	}
+
+	// If no HDR info is present, we assume sRGB
+	return GetFullRange() ? MediaShaders::YuvToRgbRec709Unscaled : MediaShaders::YuvToRgbRec709Scaled;
+}
+
+bool IElectraTextureSampleBase::GetFullRange() const
+{
+	if (auto PinnedColorimetry = Colorimetry.Pin())
+	{
+		if (auto PinnedHDRInfo = HDRInfo.Pin())
+		{
+			if (PinnedHDRInfo->GetHDRType() != IVideoDecoderHDRInformation::EType::Unknown)
+			{
+				// For HDR we assume full range at all times
+				return true;
+			}
+		}
+		// SDR honors the flag in the MPEG stream...
+		return (PinnedColorimetry->GetMPEGDefinition()->VideoFullRangeFlag != 0);
+	}
+	// SDR with no MPEG info is assumed to be video range
+	return false;
+}
+
+FMatrix44f IElectraTextureSampleBase::GetSampleToRGBMatrix() const
+{
+	return YuvToRgbMtx;
+}
+
+FMatrix44f IElectraTextureSampleBase::GetGamutToXYZMatrix() const
+{
+	if (auto PinnedHDRInfo = HDRInfo.Pin())
+	{
+		switch (PinnedHDRInfo->GetHDRType())
+		{
+		case IVideoDecoderHDRInformation::EType::PQ10:
+		case IVideoDecoderHDRInformation::EType::HLG10:
+		case IVideoDecoderHDRInformation::EType::HDR10:		return GamutToXYZMatrix(EDisplayColorGamut::Rec2020_D65);
+		case IVideoDecoderHDRInformation::EType::Unknown:	return GamutToXYZMatrix(EDisplayColorGamut::sRGB_D65);
+		}
+	}
+
+	// If no HDR info is present, we assume sRGB
+	return GamutToXYZMatrix(EDisplayColorGamut::sRGB_D65);
+}
+
+FVector2f IElectraTextureSampleBase::GetWhitePoint() const
+{
+	if (auto PinnedHDRInfo = HDRInfo.Pin())
+	{
+		if (auto ColorVolume = PinnedHDRInfo->GetMasteringDisplayColourVolume())
+		{
+			return FVector2f(ColorVolume->white_point_x, ColorVolume->white_point_y);
+		}
+	}
+	// If no HDR info is present, we assume sRGB
+	return FVector2f(UE::Color::GetWhitePoint(UE::Color::EWhitePoint::CIE1931_D65));
+}
+
+UE::Color::EEncoding IElectraTextureSampleBase::GetEncodingType() const
+{
+	if (auto PinnedHDRInfo = HDRInfo.Pin())
+	{
+		switch (PinnedHDRInfo->GetHDRType())
+		{
+		case IVideoDecoderHDRInformation::EType::PQ10:
+		case IVideoDecoderHDRInformation::EType::HDR10:		return UE::Color::EEncoding::ST2084;
+		case IVideoDecoderHDRInformation::EType::Unknown:	return UE::Color::EEncoding::sRGB;
+		case IVideoDecoderHDRInformation::EType::HLG10:
+			{
+			check(!"*** Implement support for HLG10 in UE::Color!");
+			return UE::Color::EEncoding::sRGB;
+			}
+		}
+	}
+	// If no HDR info is present, we assume sRGB
+	return UE::Color::EEncoding::sRGB;
+}
+
+#endif
diff --git a/Engine/Plugins/Media/ElectraUtil/Source/ElectraSamples/Private/ElectraTextureSamplePC.cpp b/Engine/Plugins/Media/ElectraUtil/Source/ElectraSamples/Private/ElectraTextureSamplePC.cpp
--- a/Engine/Plugins/Media/ElectraUtil/Source/ElectraSamples/Private/ElectraTextureSamplePC.cpp
+++ b/Engine/Plugins/Media/ElectraUtil/Source/ElectraSamples/Private/ElectraTextureSamplePC.cpp
@@ -39,69 +39,16 @@ DECLARE_GPU_STAT_NAMED(MediaWinDecoder_Convert, TEXT("MediaWinDecoder_Convert"))
 
 void FElectraTextureSample::Initialize(FVideoDecoderOutput *InVideoDecoderOutput)
 {
-	VideoDecoderOutput = StaticCastSharedPtr<FVideoDecoderOutputPC, IDecoderOutputPoolable, ESPMode::ThreadSafe>(InVideoDecoderOutput->AsShared());
+	IElectraTextureSampleBase::Initialize(InVideoDecoderOutput);
+	VideoDecoderOutputPC = static_cast<FVideoDecoderOutputPC*>(InVideoDecoderOutput);
 	SampleFormat = (VideoDecoderOutput->GetFormat() == PF_NV12) ? EMediaTextureSampleFormat::CharNV12 : EMediaTextureSampleFormat::P010;
-	HDRInfo = VideoDecoderOutput->GetHDRInformation();
-	Colorimetry = VideoDecoderOutput->GetColorimetry();
-
-	bool bFullRange = false; 
-	if (auto PinnedColorimetry = Colorimetry.Pin())
-	{
-		bFullRange = (PinnedColorimetry->GetMPEGDefinition()->VideoFullRangeFlag != 0);
-	}
-
-	// Prepare YUV -> RGB matrix containing all necessary offsets and scales to produce RGB straight from sample data
-	const FMatrix* Mtx = bFullRange ? &MediaShaders::YuvToRgbRec709Unscaled : &MediaShaders::YuvToRgbRec709Scaled;
-	FVector Off = (SampleFormat == EMediaTextureSampleFormat::P010) ? (bFullRange ? MediaShaders::YUVOffsetNoScale16bits : MediaShaders::YUVOffset16bits)
-																	: (bFullRange ? MediaShaders::YUVOffsetNoScale8bits : MediaShaders::YUVOffset8bits);
-	float Scale = 1.0f;
-
-	if (auto PinnedHDRInfo = HDRInfo.Pin())
-	{
-		switch (PinnedHDRInfo->GetHDRType())
-		{
-		case IVideoDecoderHDRInformation::EType::PQ10:
-		case IVideoDecoderHDRInformation::EType::HLG10:
-		case IVideoDecoderHDRInformation::EType::HDR10:
-			Mtx = &MediaShaders::YuvToRgbRec2020Unscaled;
-			Off = MediaShaders::YUVOffsetNoScale16bits;
-			break;
-		default:
-			break;
-		}
-	}
-
-	// Matrix to transform sample data to standard YUV values
-	FMatrix PreMtx = FMatrix::Identity;
-	PreMtx.M[0][0] = Scale;
-	PreMtx.M[1][1] = Scale;
-	PreMtx.M[2][2] = Scale;
-	PreMtx.M[0][3] = -Off.X;
-	PreMtx.M[1][3] = -Off.Y;
-	PreMtx.M[2][3] = -Off.Z;
-
-	// Combine this with the actual YUV-RGB conversion
-	YuvToRgbMtx = FMatrix44f(*Mtx * PreMtx);
 }
 
-
-#if !UE_SERVER
-void FElectraTextureSample::InitializePoolable()
-{
-}
-
-void FElectraTextureSample::ShutdownPoolable()
-{
-	VideoDecoderOutput.Reset();
-}
-#endif
-
-
 IMediaTextureSampleConverter* FElectraTextureSample::GetMediaTextureSampleConverter()
 {
 #ifdef ELECTRA_HAVE_DX11
 	// All versions might need SW fallback - check if we have a real texture as source -> converter needed
-	return (VideoDecoderOutput && VideoDecoderOutput->GetTexture()) ? this : nullptr;
+	return (VideoDecoderOutputPC && VideoDecoderOutputPC->GetTexture()) ? this : nullptr;
 #else
 	return nullptr;
 #endif
@@ -218,14 +165,14 @@ bool FElectraTextureSample::Convert(FTexture2DRHIRef & InDstTexture, const FConv
 
 	bool bCrossDeviceCopy;
 	EPixelFormat Format;
-	if (VideoDecoderOutput->GetOutputType() != FVideoDecoderOutputPC::EOutputType::HardwareWin8Plus)
+	if (VideoDecoderOutputPC->GetOutputType() != FVideoDecoderOutputPC::EOutputType::HardwareWin8Plus)
 	{
-		if (VideoDecoderOutput->GetOutputType() != FVideoDecoderOutputPC::EOutputType::Hardware_DX)
+		if (VideoDecoderOutputPC->GetOutputType() != FVideoDecoderOutputPC::EOutputType::Hardware_DX)
 		{
 			//
 			// SW decoder has decoded into a HW texture (not known to RHI) -> copy it into an RHI one
 			//
-			check(VideoDecoderOutput->GetOutputType() == FVideoDecoderOutputPC::EOutputType::SoftwareWin8Plus);
+			check(VideoDecoderOutputPC->GetOutputType() == FVideoDecoderOutputPC::EOutputType::SoftwareWin8Plus);
 			check(VideoDecoderOutput->GetFormat() == EPixelFormat::PF_NV12);
 
 			Format = EPixelFormat::PF_G8; // use fixed format: we flag this as NV12, too - as it is - but DX11 will only support a "higher than normal G8 texture" (with specialized access in shader)
@@ -262,12 +209,12 @@ bool FElectraTextureSample::Convert(FTexture2DRHIRef & InDstTexture, const FConv
 	// Copy data into RHI texture
 	if (RHICmdList.Bypass())
 	{
-		FRHICommandCopyResource Cmd(VideoDecoderOutput->GetTexture(), Texture, bCrossDeviceCopy);
+		FRHICommandCopyResource Cmd(VideoDecoderOutputPC->GetTexture(), Texture, bCrossDeviceCopy);
 		Cmd.Execute(RHICmdList);
 	}
 	else
 	{
-		new (RHICmdList.AllocCommand<FRHICommandCopyResource>()) FRHICommandCopyResource(VideoDecoderOutput->GetTexture(), Texture, bCrossDeviceCopy);
+		new (RHICmdList.AllocCommand<FRHICommandCopyResource>()) FRHICommandCopyResource(VideoDecoderOutputPC->GetTexture(), Texture, bCrossDeviceCopy);
 	}
 	
 	return true;
@@ -281,7 +228,7 @@ const void* FElectraTextureSample::GetBuffer()
 {
 	if (VideoDecoderOutput)
 	{
-		return VideoDecoderOutput->GetBuffer().GetData();
+		return VideoDecoderOutputPC->GetBuffer().GetData();
 	}
 	return nullptr;
 }
@@ -291,152 +238,19 @@ uint32 FElectraTextureSample::GetStride() const
 {
 	if (VideoDecoderOutput)
 	{
-		return VideoDecoderOutput->GetStride();
+		return VideoDecoderOutputPC->GetStride();
 	}
 	return 0;
 }
 
 
-FIntPoint FElectraTextureSample::GetDim() const
-{
-	if (VideoDecoderOutput)
-	{
-		return VideoDecoderOutput->GetDim();
-	}
-	return FIntPoint::ZeroValue;
-}
-
-
-FIntPoint FElectraTextureSample::GetOutputDim() const
-{
-	if (VideoDecoderOutput)
-	{
-		return VideoDecoderOutput->GetOutputDim();
-	}
-	return FIntPoint::ZeroValue;
-}
-
-
-FMediaTimeStamp FElectraTextureSample::GetTime() const
-{
-	if (VideoDecoderOutput)
-	{
-		const FDecoderTimeStamp TimeStamp = VideoDecoderOutput->GetTime();
-		return FMediaTimeStamp(TimeStamp.Time, TimeStamp.SequenceIndex);
-	}
-	return FMediaTimeStamp();
-}
-
-
-FTimespan FElectraTextureSample::GetDuration() const
-{
-	if (VideoDecoderOutput)
-	{
-		return VideoDecoderOutput->GetDuration();
-	}
-	return FTimespan::Zero();
-}
-
 IMFSample* FElectraTextureSample::GetMFSample()
 {
 	if (VideoDecoderOutput)
 	{
-		return VideoDecoderOutput->GetMFSample().GetReference();
+		return VideoDecoderOutputPC->GetMFSample().GetReference();
 	}
 	return nullptr;
 }
 
-
-bool FElectraTextureSample::IsOutputSrgb() const
-{
-	if (auto PinnedHDRInfo = HDRInfo.Pin())
-	{
-		// If the HDR type is unknown we also assume sRGB
-		return (PinnedHDRInfo->GetHDRType() == IVideoDecoderHDRInformation::EType::Unknown);
-	}
-	// If no HDR info is present, we assume sRGB
-	return true;
-}
-
-const FMatrix& FElectraTextureSample::GetYUVToRGBMatrix() const
-{
-	if (auto PinnedHDRInfo = HDRInfo.Pin())
-	{
-		switch (PinnedHDRInfo->GetHDRType())
-		{
-		case IVideoDecoderHDRInformation::EType::PQ10:
-		case IVideoDecoderHDRInformation::EType::HLG10:
-		case IVideoDecoderHDRInformation::EType::HDR10:		return MediaShaders::YuvToRgbRec2020Unscaled;
-		case IVideoDecoderHDRInformation::EType::Unknown:	break;
-		}
-	}
-
-	// If no HDR info is present, we assume sRGB
-	return GetFullRange() ? MediaShaders::YuvToRgbRec709Unscaled : MediaShaders::YuvToRgbRec709Scaled;
-}
-
-bool FElectraTextureSample::GetFullRange() const
-{
-	bool bFullRange = false;
-	if (auto PinnedColorimetry = Colorimetry.Pin())
-	{
-		bFullRange = (PinnedColorimetry->GetMPEGDefinition()->VideoFullRangeFlag != 0);
-	}
-	return bFullRange;
-}
-
-FMatrix44f FElectraTextureSample::GetSampleToRGBMatrix() const
-{
-	return YuvToRgbMtx;
-}
-
-FMatrix44f FElectraTextureSample::GetGamutToXYZMatrix() const
-{
-	if (auto PinnedHDRInfo = HDRInfo.Pin())
-	{
-		switch (PinnedHDRInfo->GetHDRType())
-		{
-		case IVideoDecoderHDRInformation::EType::PQ10:
-		case IVideoDecoderHDRInformation::EType::HLG10:
-		case IVideoDecoderHDRInformation::EType::HDR10:		return GamutToXYZMatrix(EDisplayColorGamut::Rec2020_D65);
-		case IVideoDecoderHDRInformation::EType::Unknown:	return GamutToXYZMatrix(EDisplayColorGamut::sRGB_D65);
-		}
-	}
-	// If no HDR info is present, we assume sRGB
-	return GamutToXYZMatrix(EDisplayColorGamut::sRGB_D65);
-}
-
-FVector2f FElectraTextureSample::GetWhitePoint() const
-{
-	if (auto PinnedHDRInfo = HDRInfo.Pin())
-	{
-		if (auto ColorVolume = PinnedHDRInfo->GetMasteringDisplayColourVolume())
-		{
-			return FVector2f(ColorVolume->white_point_x, ColorVolume->white_point_y);
-		}
-	}
-	// If no HDR info is present, we assume sRGB
-	return FVector2f(UE::Color::GetWhitePoint(UE::Color::EWhitePoint::CIE1931_D65));
-}
-
-UE::Color::EEncoding FElectraTextureSample::GetEncodingType() const
-{
-	if (auto PinnedHDRInfo = HDRInfo.Pin())
-	{
-		switch (PinnedHDRInfo->GetHDRType())
-		{
-		case IVideoDecoderHDRInformation::EType::PQ10:
-		case IVideoDecoderHDRInformation::EType::HDR10:		return UE::Color::EEncoding::ST2084;
-		case IVideoDecoderHDRInformation::EType::Unknown:	return UE::Color::EEncoding::sRGB;
-		case IVideoDecoderHDRInformation::EType::HLG10:
-			{
-			check(!"*** Implement support for HLG10 in UE::Color!");
-			return UE::Color::EEncoding::sRGB;
-			}
-		}
-	}
-	// If no HDR info is present, we assume sRGB
-	return UE::Color::EEncoding::sRGB;
-}
-
 #endif
diff --git a/Engine/Plugins/Media/ElectraUtil/Source/ElectraSamples/Private/Java/src/com/epicgames/unreal/ElectraTextureSample.java b/Engine/Plugins/Media/ElectraUtil/Source/ElectraSamples/Private/Java/src/com/epicgames/unreal/ElectraTextureSample.java
--- a/Engine/Plugins/Media/ElectraUtil/Source/ElectraSamples/Private/Java/src/com/epicgames/unreal/ElectraTextureSample.java
+++ b/Engine/Plugins/Media/ElectraUtil/Source/ElectraSamples/Private/Java/src/com/epicgames/unreal/ElectraTextureSample.java
@@ -52,12 +52,12 @@ public class ElectraTextureSample
 		}
 	}
 
-	public FFrameUpdateInfo GetVideoFrameUpdateInfo(int destTextureES, int width, int height)
+	public FFrameUpdateInfo GetVideoFrameUpdateInfo(int destTextureES, int width, int height, boolean bIs10Bbit)
 	{
 		FFrameUpdateInfo Info = null;
 		if (BitmapRenderer != null)
 		{
-			Info = BitmapRenderer.GetVideoFrameUpdateInfo(destTextureES, width, height);
+			Info = BitmapRenderer.GetVideoFrameUpdateInfo(destTextureES, width, height, bIs10Bbit);
 		}
 		return Info;
 	}
@@ -89,12 +89,11 @@ public class ElectraTextureSample
 		private int mTextureID = -1;
 		private int mFBO = -1;
 		private int mBlitVertexShaderID = -1;
-		private int mBlitFragmentShaderID = -1;
+		private int[] mBlitFragmentShaderID = new int[2];
 		private float[] mTransformMatrix = new float[16];
 		private boolean mTriangleVerticesDirty = true;
 		private boolean mUseOwnContext = true;
 		private boolean mVulkanRenderer = false;
-		private boolean mSwizzlePixels = false;
 
 		private int GL_TEXTURE_EXTERNAL_OES = 0x8D65;
 
@@ -119,15 +118,9 @@ public class ElectraTextureSample
 			mEglContext = EGL14.EGL_NO_CONTEXT;
 			mUseOwnContext = true;
 
-			if (mVulkanRenderer)
-			{
-				mSwizzlePixels = true;
-
-				Log.d(TAG,"FBitmapRenderer: mSwizzlePixels");
-			}
-			else
+			if (!mVulkanRenderer)
 			{
-				String RendererString = GLES20.glGetString(GLES20.GL_RENDERER);
+				String RendererString = GLES32.glGetString(GLES32.GL_RENDERER);
 				Log.d(TAG,"FBitmapRenderer: GLES: " + RendererString);
 				// Do not use shared context if Adreno before 400 or on older Android than Marshmallow
 				if (RendererString.contains("Adreno (TM) "))
@@ -273,7 +266,7 @@ public class ElectraTextureSample
 		private void initSurfaceTexture()
 		{
 			int[] textures = new int[1];
-			GLES20.glGenTextures(1, textures, 0);
+			GLES32.glGenTextures(1, textures, 0);
 			mTextureID = textures[0];
 			if (mTextureID <= 0)
 			{
@@ -286,7 +279,7 @@ public class ElectraTextureSample
 
 			int[] glInt = new int[1];
 
-			GLES20.glGenFramebuffers(1,glInt,0);
+			GLES32.glGenFramebuffers(1,glInt,0);
 			mFBO = glInt[0];
 			if (mFBO <= 0)
 			{
@@ -295,48 +288,65 @@ public class ElectraTextureSample
 				return;
 			}
 			// Special shaders for blit of movie texture.
-			mBlitVertexShaderID = createShader(GLES20.GL_VERTEX_SHADER, mBlitVextexShader);
+			mBlitVertexShaderID = createShader(GLES32.GL_VERTEX_SHADER, mBlitVextexShader);
 			if (mBlitVertexShaderID == 0)
 			{
 				Log.e(TAG,"mBlitVertexShaderID == 0");
 				release();
 				return;
 			}
-			int mBlitFragmentShaderID = createShader(GLES20.GL_FRAGMENT_SHADER,
-				mVulkanRenderer ? (mSwizzlePixels ? mBlitFragmentShaderBGRA_NoConv : mBlitFragmentShaderRGBA_NoConv)
-								: (mSwizzlePixels ? mBlitFragmentShaderBGRA : mBlitFragmentShaderRGBA));
-			if (mBlitFragmentShaderID == 0)
-			{
-				Log.e(TAG,"mBlitFragmentShaderID == 0");
+			//
+			// Android will deliver the data as "image external" to ES.
+			// - Anything read from such a texture will "magically" appear as [0..1] range RGB data in the shader
+			// - For SDR material this will be sRGB data (still very much with gamma curve applied)
+			// - For HDR material thisd will be PQ data (here, too, the EOTF is still very much applied)
+			// - CS will be whatever we have in the material playing (mostly Rec.709 (SDR) or Rec.2020 (HDR))
+			//
+			// -> Hence this indeed complies with the old [0..1] range for "image external". We just need to be very careful
+			//    with how to interpret the data.
+			//
+			// ES:		We copy the data directly into an externally (RHI created) provided texture from where MFW can further process it as needed
+			// Vulkan:	We simply copy the data into a CPU-side buffer (any upload into a texture is done in MediaFramework - an ES texture serves as inbetween buffer as we cannot access the "image external" data)
+			//
+			int[] mBlitFragmentShaderID = new int[2];
+			mBlitFragmentShaderID[0] = createShader(GLES32.GL_FRAGMENT_SHADER, mBlitFragmentShaderRGBA_NoConv);
+			mBlitFragmentShaderID[1] = createShader(GLES32.GL_FRAGMENT_SHADER, mBlitFragmentShaderBGRA_NoConv);
+
+			if (mBlitFragmentShaderID[0] == 0 || mBlitFragmentShaderID[1] == 0)
+			{
+				Log.e(TAG,"mBlitFragmentShaderID[i] == 0");
 				release();
 				return;
 			}
-			mProgram = GLES20.glCreateProgram();
-			if (mProgram <= 0)
+			for(int i=0; i<2; ++i)
 			{
-				Log.e(TAG,"mProgram <= 0");
-				release();
-				return;
-			}
-			GLES20.glAttachShader(mProgram, mBlitVertexShaderID);
-			GLES20.glAttachShader(mProgram, mBlitFragmentShaderID);
-			GLES20.glLinkProgram(mProgram);
-			int[] linkStatus = new int[1];
-			GLES20.glGetProgramiv(mProgram, GLES20.GL_LINK_STATUS, linkStatus, 0);
-			if (linkStatus[0] != GLES20.GL_TRUE)
-			{
-				Log.e(TAG,"Could not link program: ");
-				Log.e(TAG, GLES20.glGetProgramInfoLog(mProgram));
-				GLES20.glDeleteProgram(mProgram);
-				mProgram = 0;
-				release();
-				return;
+				mProgram[i] = GLES32.glCreateProgram();
+				if (mProgram[i] <= 0)
+				{
+					Log.e(TAG,"mProgram[i] <= 0");
+					release();
+					return;
+				}
+				GLES32.glAttachShader(mProgram[i], mBlitVertexShaderID);
+				GLES32.glAttachShader(mProgram[i], mBlitFragmentShaderID[i]);
+				GLES32.glLinkProgram(mProgram[i]);
+				int[] linkStatus = new int[1];
+				GLES32.glGetProgramiv(mProgram[i], GLES32.GL_LINK_STATUS, linkStatus, 0);
+				if (linkStatus[0] != GLES32.GL_TRUE)
+				{
+					Log.e(TAG,"Could not link program: ");
+					Log.e(TAG, GLES32.glGetProgramInfoLog(mProgram[i]));
+					GLES32.glDeleteProgram(mProgram[i]);
+					mProgram[i] = 0;
+					release();
+					return;
+				}
+				mPositionAttrib = GLES32.glGetAttribLocation(mProgram[i], "Position");
+				mTexCoordsAttrib = GLES32.glGetAttribLocation(mProgram[i], "TexCoords");
+				mTextureUniform = GLES32.glGetUniformLocation(mProgram[i], "VideoTexture");
 			}
-			mPositionAttrib = GLES20.glGetAttribLocation(mProgram, "Position");
-			mTexCoordsAttrib = GLES20.glGetAttribLocation(mProgram, "TexCoords");
-			mTextureUniform = GLES20.glGetUniformLocation(mProgram, "VideoTexture");
 
-			GLES20.glGenBuffers(1,glInt,0);
+			GLES32.glGenBuffers(1,glInt,0);
 			mBlitBuffer = glInt[0];
 			if (mBlitBuffer <= 0)
 			{
@@ -354,13 +364,13 @@ public class ElectraTextureSample
 			// Set up GL state
 			if (mUseOwnContext)
 			{
-				GLES20.glDisable(GLES20.GL_BLEND);
-				GLES20.glDisable(GLES20.GL_CULL_FACE);
-				GLES20.glDisable(GLES20.GL_SCISSOR_TEST);
-				GLES20.glDisable(GLES20.GL_STENCIL_TEST);
-				GLES20.glDisable(GLES20.GL_DEPTH_TEST);
-				GLES20.glDisable(GLES20.GL_DITHER);
-				GLES20.glColorMask(true,true,true,true);
+				GLES32.glDisable(GLES32.GL_BLEND);
+				GLES32.glDisable(GLES32.GL_CULL_FACE);
+				GLES32.glDisable(GLES32.GL_SCISSOR_TEST);
+				GLES32.glDisable(GLES32.GL_STENCIL_TEST);
+				GLES32.glDisable(GLES32.GL_DEPTH_TEST);
+				GLES32.glDisable(GLES32.GL_DITHER);
+				GLES32.glColorMask(true,true,true,true);
 			}
 		}
 
@@ -377,16 +387,16 @@ public class ElectraTextureSample
 
 			// save VBO state
 			int[] glInt = new int[1];
-			GLES20.glGetIntegerv(GLES20.GL_ARRAY_BUFFER_BINDING, glInt, 0);
+			GLES32.glGetIntegerv(GLES32.GL_ARRAY_BUFFER_BINDING, glInt, 0);
 			int previousVBO = glInt[0];
 
-			GLES20.glBindBuffer(GLES20.GL_ARRAY_BUFFER, mBlitBuffer);
-			GLES20.glBufferData(GLES20.GL_ARRAY_BUFFER,
+			GLES32.glBindBuffer(GLES32.GL_ARRAY_BUFFER, mBlitBuffer);
+			GLES32.glBufferData(GLES32.GL_ARRAY_BUFFER,
 				mTriangleVerticesData.length*FLOAT_SIZE_BYTES,
-				mTriangleVertices, GLES20.GL_STATIC_DRAW);
+				mTriangleVertices, GLES32.GL_STATIC_DRAW);
 
 			// restore VBO state
-			GLES20.glBindBuffer(GLES20.GL_ARRAY_BUFFER, previousVBO);
+			GLES32.glBindBuffer(GLES32.GL_ARRAY_BUFFER, previousVBO);
 
 			mTriangleVerticesDirty = false;
 		}
@@ -398,18 +408,18 @@ public class ElectraTextureSample
 
 		private int createShader(int shaderType, String source)
 		{
-			int shader = GLES20.glCreateShader(shaderType);
+			int shader = GLES32.glCreateShader(shaderType);
 			if (shader != 0)
 			{
-				GLES20.glShaderSource(shader, source);
-				GLES20.glCompileShader(shader);
+				GLES32.glShaderSource(shader, source);
+				GLES32.glCompileShader(shader);
 				int[] compiled = new int[1];
-				GLES20.glGetShaderiv(shader, GLES20.GL_COMPILE_STATUS, compiled, 0);
+				GLES32.glGetShaderiv(shader, GLES32.GL_COMPILE_STATUS, compiled, 0);
 				if (compiled[0] == 0)
 				{
 					Log.e(TAG,"Could not compile shader " + shaderType + ":");
-					Log.e(TAG, GLES20.glGetShaderInfoLog(shader));
-					GLES20.glDeleteShader(shader);
+					Log.e(TAG, GLES32.glGetShaderInfoLog(shader));
+					GLES32.glDeleteShader(shader);
 					shader = 0;
 				}
 			}
@@ -454,35 +464,13 @@ public class ElectraTextureSample
 			"	gl_Position = vec4(Position, 0.0, 1.0);\n" +
 			"}\n";
 
-		private final String mBlitFragmentShaderBGRA =
-			"#extension GL_OES_EGL_image_external : require\n" +
-			"uniform samplerExternalOES VideoTexture;\n" +
-			"varying highp vec2 TexCoord;\n" +
-			"void main()\n" +
-			"{\n" +
-			"	highp vec4 Color = texture2D(VideoTexture, TexCoord).bgra;\n" +
-			"	Color.rgb = max(Color.rgb, 6.10352e-5);\n" +
-			"	Color.rgb = mix(pow(Color.rgb * vec3(1.0 / 1.055) + vec3(0.0521327), vec3(2.4) ), Color.rgb * vec3(1.0 / 12.92), vec3(lessThanEqual(Color.rgb, vec3(0.04045))));\n" +
-			"	gl_FragColor = Color;\n" +
-			"}\n";
-		private final String mBlitFragmentShaderRGBA =
-			"#extension GL_OES_EGL_image_external : require\n" +
-			"uniform samplerExternalOES VideoTexture;\n" +
-			"varying highp vec2 TexCoord;\n" +
-			"void main()\n" +
-			"{\n" +
-			"	highp vec4 Color = texture2D(VideoTexture, TexCoord).rgba;\n" +
-			"	Color.rgb = max(Color.rgb, 6.10352e-5);\n" +
-			"	Color.rgb = mix(pow(Color.rgb * vec3(1.0 / 1.055) + vec3(0.0521327), vec3(2.4) ), Color.rgb * vec3(1.0 / 12.92), vec3(lessThanEqual(Color.rgb, vec3(0.04045))));\n" +
-			"	gl_FragColor = Color;\n" +
-			"}\n";
 		private final String mBlitFragmentShaderBGRA_NoConv =
 			"#extension GL_OES_EGL_image_external : require\n" +
 			"uniform samplerExternalOES VideoTexture;\n" +
 			"varying highp vec2 TexCoord;\n" +
 			"void main()\n" +
 			"{\n" +
-			"	gl_FragColor = texture2D(VideoTexture, TexCoord).bgra;;\n" +
+			"	gl_FragColor = texture2D(VideoTexture, TexCoord).bgra;\n" +
 			"}\n";
 		private final String mBlitFragmentShaderRGBA_NoConv =
 			"#extension GL_OES_EGL_image_external : require\n" +
@@ -490,10 +478,10 @@ public class ElectraTextureSample
 			"varying highp vec2 TexCoord;\n" +
 			"void main()\n" +
 			"{\n" +
-			"	gl_FragColor = texture2D(VideoTexture, TexCoord).bgra;;\n" +
+			"	gl_FragColor = texture2D(VideoTexture, TexCoord).rgba;\n" +
 			"}\n";
 
-		private int mProgram;
+		private int[] mProgram = new int[2];
 		private int mPositionAttrib;
 		private int mTexCoordsAttrib;
 		private int mBlitBuffer;
@@ -502,7 +490,7 @@ public class ElectraTextureSample
 		//
 		// Vulkan version: resulting image is returned as byte array in frame info data
 		//
-		public FFrameUpdateInfo updateFrameData()
+		public FFrameUpdateInfo updateFrameData(boolean bIs10Bit)
 		{
 			synchronized(this)
 			{
@@ -517,7 +505,7 @@ public class ElectraTextureSample
 					mFrameDataSize = frameDataSize;
 				}
 				// Copy surface texture to frame data.
-				if (!copyFrameTexture(0, mFrameData))
+				if (!copyFrameTexture(0, mFrameData, bIs10Bit))
 				{
 					return null;
 				}
@@ -542,12 +530,12 @@ public class ElectraTextureSample
 		//
 		// ES version: resulting image is written to given texture
 		//
-		public FFrameUpdateInfo updateFrameData(int destTexture)
+		public FFrameUpdateInfo updateFrameData(int destTexture, boolean bIs10Bit)
 		{
 			synchronized(this)
 			{
 				// Copy surface texture to destination texture.
-				if (!copyFrameTexture(destTexture, null))
+				if (!copyFrameTexture(destTexture, null, bIs10Bit))
 				{
 					return null;
 				}
@@ -574,7 +562,7 @@ public class ElectraTextureSample
 
 		// Copy the surface texture to another texture, or to raw data.
 		// Note: copying to raw data creates a temporary FBO texture.
-		private boolean copyFrameTexture(int destTexture, java.nio.Buffer destData)
+		private boolean copyFrameTexture(int destTexture, java.nio.Buffer destData, boolean bIs10Bit)
 		{
 			if (null == mSurfaceTexture)
 			{
@@ -591,10 +579,10 @@ public class ElectraTextureSample
 			if (mUseOwnContext)
 			{
 				// Received reports of these not being preserved when changing contexts
-				GLES20.glActiveTexture(GLES20.GL_TEXTURE0);
-				GLES20.glGetTexParameteriv(GLES20.GL_TEXTURE_2D, GLES20.GL_TEXTURE_MIN_FILTER, glInt, 0);
+				GLES32.glActiveTexture(GLES32.GL_TEXTURE0);
+				GLES32.glGetTexParameteriv(GLES32.GL_TEXTURE_2D, GLES32.GL_TEXTURE_MIN_FILTER, glInt, 0);
 				previousMinFilter = glInt[0];
-				GLES20.glGetTexParameteriv(GLES20.GL_TEXTURE_2D, GLES20.GL_TEXTURE_MAG_FILTER, glInt, 0);
+				GLES32.glGetTexParameteriv(GLES32.GL_TEXTURE_2D, GLES32.GL_TEXTURE_MAG_FILTER, glInt, 0);
 				previousMagFilter = glInt[0];
 
 				saveContext();
@@ -603,24 +591,24 @@ public class ElectraTextureSample
 			else
 			{
 				// Clear gl errors as they can creep in from the UE4 renderer.
-				GLES20.glGetError();
-
-				previousBlend = GLES20.glIsEnabled(GLES20.GL_BLEND);
-				previousCullFace = GLES20.glIsEnabled(GLES20.GL_CULL_FACE);
-				previousScissorTest = GLES20.glIsEnabled(GLES20.GL_SCISSOR_TEST);
-				previousStencilTest = GLES20.glIsEnabled(GLES20.GL_STENCIL_TEST);
-				previousDepthTest = GLES20.glIsEnabled(GLES20.GL_DEPTH_TEST);
-				previousDither = GLES20.glIsEnabled(GLES20.GL_DITHER);
-				GLES20.glGetIntegerv(GLES20.GL_FRAMEBUFFER_BINDING, glInt, 0);
+				GLES32.glGetError();
+
+				previousBlend = GLES32.glIsEnabled(GLES32.GL_BLEND);
+				previousCullFace = GLES32.glIsEnabled(GLES32.GL_CULL_FACE);
+				previousScissorTest = GLES32.glIsEnabled(GLES32.GL_SCISSOR_TEST);
+				previousStencilTest = GLES32.glIsEnabled(GLES32.GL_STENCIL_TEST);
+				previousDepthTest = GLES32.glIsEnabled(GLES32.GL_DEPTH_TEST);
+				previousDither = GLES32.glIsEnabled(GLES32.GL_DITHER);
+				GLES32.glGetIntegerv(GLES32.GL_FRAMEBUFFER_BINDING, glInt, 0);
 				previousFBO = glInt[0];
-				GLES20.glGetIntegerv(GLES20.GL_ARRAY_BUFFER_BINDING, glInt, 0);
+				GLES32.glGetIntegerv(GLES32.GL_ARRAY_BUFFER_BINDING, glInt, 0);
 				previousVBO = glInt[0];
-				GLES20.glGetIntegerv(GLES20.GL_VIEWPORT, previousViewport, 0);
+				GLES32.glGetIntegerv(GLES32.GL_VIEWPORT, previousViewport, 0);
 
-				GLES20.glActiveTexture(GLES20.GL_TEXTURE0);
-				GLES20.glGetTexParameteriv(GLES20.GL_TEXTURE_2D, GLES20.GL_TEXTURE_MIN_FILTER, glInt, 0);
+				GLES32.glActiveTexture(GLES32.GL_TEXTURE0);
+				GLES32.glGetTexParameteriv(GLES32.GL_TEXTURE_2D, GLES32.GL_TEXTURE_MIN_FILTER, glInt, 0);
 				previousMinFilter = glInt[0];
-				GLES20.glGetTexParameteriv(GLES20.GL_TEXTURE_2D, GLES20.GL_TEXTURE_MAG_FILTER, glInt, 0);
+				GLES32.glGetTexParameteriv(GLES32.GL_TEXTURE_2D, GLES32.GL_TEXTURE_MAG_FILTER, glInt, 0);
 				previousMagFilter = glInt[0];
 
 				glVerify("save state");
@@ -655,18 +643,18 @@ public class ElectraTextureSample
 
 			if (!mUseOwnContext)
 			{
-				GLES20.glDisable(GLES20.GL_BLEND);
-				GLES20.glDisable(GLES20.GL_CULL_FACE);
-				GLES20.glDisable(GLES20.GL_SCISSOR_TEST);
-				GLES20.glDisable(GLES20.GL_STENCIL_TEST);
-				GLES20.glDisable(GLES20.GL_DEPTH_TEST);
-				GLES20.glDisable(GLES20.GL_DITHER);
-				GLES20.glColorMask(true,true,true,true);
+				GLES32.glDisable(GLES32.GL_BLEND);
+				GLES32.glDisable(GLES32.GL_CULL_FACE);
+				GLES32.glDisable(GLES32.GL_SCISSOR_TEST);
+				GLES32.glDisable(GLES32.GL_STENCIL_TEST);
+				GLES32.glDisable(GLES32.GL_DEPTH_TEST);
+				GLES32.glDisable(GLES32.GL_DITHER);
+				GLES32.glColorMask(true,true,true,true);
 
 				glVerify("reset state");
 			}
 
-			GLES20.glViewport(0, 0, mTextureWidth, mTextureHeight);
+			GLES32.glViewport(0, 0, mTextureWidth, mTextureHeight);
 
 			glVerify("set viewport");
 
@@ -675,7 +663,7 @@ public class ElectraTextureSample
 			if (null != destData)
 			{
 				// Create temporary FBO for data copy.
-				GLES20.glGenTextures(1,glInt,0);
+				GLES32.glGenTextures(1,glInt,0);
 				FBOTextureID = glInt[0];
 			}
 			else
@@ -684,77 +672,85 @@ public class ElectraTextureSample
 				FBOTextureID = destTexture;
 			}
 			// Set the FBO to draw into the texture one-to-one.
-			GLES20.glBindTexture(GLES20.GL_TEXTURE_2D, FBOTextureID);
-			GLES20.glTexParameteri(GLES20.GL_TEXTURE_2D,
-				GLES20.GL_TEXTURE_MIN_FILTER, GLES20.GL_NEAREST);
-			GLES20.glTexParameteri(GLES20.GL_TEXTURE_2D,
-				GLES20.GL_TEXTURE_MAG_FILTER, GLES20.GL_NEAREST);
-			GLES20.glTexParameteri(GLES20.GL_TEXTURE_2D,
-				GLES20.GL_TEXTURE_WRAP_S, GLES20.GL_CLAMP_TO_EDGE);
-			GLES20.glTexParameteri(GLES20.GL_TEXTURE_2D,
-				GLES20.GL_TEXTURE_WRAP_T, GLES20.GL_CLAMP_TO_EDGE);
+			GLES32.glBindTexture(GLES32.GL_TEXTURE_2D, FBOTextureID);
+			GLES32.glTexParameteri(GLES32.GL_TEXTURE_2D,
+				GLES32.GL_TEXTURE_MIN_FILTER, GLES32.GL_NEAREST);
+			GLES32.glTexParameteri(GLES32.GL_TEXTURE_2D,
+				GLES32.GL_TEXTURE_MAG_FILTER, GLES32.GL_NEAREST);
+			GLES32.glTexParameteri(GLES32.GL_TEXTURE_2D,
+				GLES32.GL_TEXTURE_WRAP_S, GLES32.GL_CLAMP_TO_EDGE);
+			GLES32.glTexParameteri(GLES32.GL_TEXTURE_2D,
+				GLES32.GL_TEXTURE_WRAP_T, GLES32.GL_CLAMP_TO_EDGE);
 			// Create the temp FBO data if needed.
 			if (null != destData)
 			{
-				GLES20.glTexImage2D(GLES20.GL_TEXTURE_2D, 0,
-					GLES20.GL_RGBA,
+				GLES32.glTexImage2D(GLES32.GL_TEXTURE_2D, 0,
+					bIs10Bit ? GLES32.GL_RGB10_A2 : GLES32.GL_RGBA,
 					mTextureWidth, mTextureHeight,
-					0, GLES20.GL_RGBA, GLES20.GL_UNSIGNED_BYTE, null);
+					0,
+					GLES32.GL_RGBA,
+					bIs10Bit ? GLES32.GL_UNSIGNED_INT_2_10_10_10_REV : GLES32.GL_UNSIGNED_BYTE, null);
 			}
 
 			glVerify("set-up FBO texture");
 			// Set to render to the FBO.
-			GLES20.glBindFramebuffer(GLES20.GL_FRAMEBUFFER, mFBO);
+			GLES32.glBindFramebuffer(GLES32.GL_FRAMEBUFFER, mFBO);
 			glVerify("glBindFramebuffer");
-			GLES20.glFramebufferTexture2D(
-				GLES20.GL_FRAMEBUFFER,
-				GLES20.GL_COLOR_ATTACHMENT0,
-				GLES20.GL_TEXTURE_2D, FBOTextureID, 0);
+			GLES32.glFramebufferTexture2D(
+				GLES32.GL_FRAMEBUFFER,
+				GLES32.GL_COLOR_ATTACHMENT0,
+				GLES32.GL_TEXTURE_2D, FBOTextureID, 0);
 
 			// check status
-			int status = GLES20.glCheckFramebufferStatus(GLES20.GL_FRAMEBUFFER);
-			if (status != GLES20.GL_FRAMEBUFFER_COMPLETE)
+			int status = GLES32.glCheckFramebufferStatus(GLES32.GL_FRAMEBUFFER);
+			if (status != GLES32.GL_FRAMEBUFFER_COMPLETE)
 			{
 				Log.w(TAG,"Failed to complete framebuffer attachment ("+status+")");
 			}
 
 			// The special shaders to render from the video texture.
-			GLES20.glUseProgram(mProgram);
+			// (Vulkan SDR needs to convert BGR / RGB)
+			GLES32.glUseProgram((!bIs10Bit && mVulkanRenderer) ? mProgram[1] : mProgram[0]);
 
 			// Set the mesh that renders the video texture.
 			UpdateVertexData();
-			GLES20.glBindBuffer(GLES20.GL_ARRAY_BUFFER, mBlitBuffer);
-			GLES20.glEnableVertexAttribArray(mPositionAttrib);
-			GLES20.glVertexAttribPointer(mPositionAttrib, 2, GLES20.GL_FLOAT, false,
+			GLES32.glBindBuffer(GLES32.GL_ARRAY_BUFFER, mBlitBuffer);
+			GLES32.glEnableVertexAttribArray(mPositionAttrib);
+			GLES32.glVertexAttribPointer(mPositionAttrib, 2, GLES32.GL_FLOAT, false,
 				TRIANGLE_VERTICES_DATA_STRIDE_BYTES, 0);
-			GLES20.glEnableVertexAttribArray(mTexCoordsAttrib);
-			GLES20.glVertexAttribPointer(mTexCoordsAttrib, 2, GLES20.GL_FLOAT, false,
+			GLES32.glEnableVertexAttribArray(mTexCoordsAttrib);
+			GLES32.glVertexAttribPointer(mTexCoordsAttrib, 2, GLES32.GL_FLOAT, false,
 				TRIANGLE_VERTICES_DATA_STRIDE_BYTES,
 				TRIANGLE_VERTICES_DATA_UV_OFFSET*FLOAT_SIZE_BYTES);
 
 			glVerify("setup movie texture read");
 
-//GLES20.glClearColor(1.0f, 1.0f, 1.0f, 1.0f);
-//GLES20.glClear( GLES20.GL_COLOR_BUFFER_BIT);
+//GLES32.glClearColor(1.0f, 1.0f, 1.0f, 1.0f);
+//GLES32.glClear( GLES32.GL_COLOR_BUFFER_BIT);
 
 			// connect 'VideoTexture' to video source texture (mTextureID).
 			// mTextureID is bound to GL_TEXTURE_EXTERNAL_OES in updateTexImage
-			GLES20.glUniform1i(mTextureUniform, 0);
-			GLES20.glActiveTexture(GLES20.GL_TEXTURE0);
-			GLES20.glBindTexture(GL_TEXTURE_EXTERNAL_OES, mTextureID);
+			GLES32.glUniform1i(mTextureUniform, 0);
+			GLES32.glActiveTexture(GLES32.GL_TEXTURE0);
+			GLES32.glBindTexture(GL_TEXTURE_EXTERNAL_OES, mTextureID);
+
+			glVerify("bind FBO movie texture read");
 
 			// Draw the video texture mesh.
-			GLES20.glDrawArrays(GLES20.GL_TRIANGLE_STRIP, 0, 4);
+			GLES32.glDrawArrays(GLES32.GL_TRIANGLE_STRIP, 0, 4);
 
 			// Make sure GL issues commands to GPU
-			GLES20.glFlush();
+			GLES32.glFlush();
+
+			glVerify("movie texture read flushed");
 
 			// Read the FBO texture pixels into raw data.
 			if (null != destData)
 			{
-				GLES20.glReadPixels(
+				GLES32.glReadPixels(
 					0, 0, mTextureWidth, mTextureHeight,
-					GLES20.GL_RGBA, GLES20.GL_UNSIGNED_BYTE,
+					GLES32.GL_RGBA,
+					bIs10Bit ? GLES32.GL_UNSIGNED_INT_2_10_10_10_REV : GLES32.GL_UNSIGNED_BYTE,
 					destData);
 			}
 
@@ -763,47 +759,47 @@ public class ElectraTextureSample
 			// Restore state and cleanup.
 			if (mUseOwnContext)
 			{
-				GLES20.glFramebufferTexture2D(
-					GLES20.GL_FRAMEBUFFER,
-					GLES20.GL_COLOR_ATTACHMENT0,
-					GLES20.GL_TEXTURE_2D, 0, 0);
+				GLES32.glFramebufferTexture2D(
+					GLES32.GL_FRAMEBUFFER,
+					GLES32.GL_COLOR_ATTACHMENT0,
+					GLES32.GL_TEXTURE_2D, 0, 0);
 
 				if (null != destData && FBOTextureID > 0)
 				{
 					glInt[0] = FBOTextureID;
-					GLES20.glDeleteTextures(1, glInt, 0);
+					GLES32.glDeleteTextures(1, glInt, 0);
 				}
 
 				restoreContext();
 
 				// Restore previous texture filtering
-				GLES20.glTexParameteri(GLES20.GL_TEXTURE_2D, GLES20.GL_TEXTURE_MIN_FILTER, previousMinFilter);
-				GLES20.glTexParameteri(GLES20.GL_TEXTURE_2D, GLES20.GL_TEXTURE_MAG_FILTER, previousMagFilter);
+				GLES32.glTexParameteri(GLES32.GL_TEXTURE_2D, GLES32.GL_TEXTURE_MIN_FILTER, previousMinFilter);
+				GLES32.glTexParameteri(GLES32.GL_TEXTURE_2D, GLES32.GL_TEXTURE_MAG_FILTER, previousMagFilter);
 			}
 			else
 			{
-				GLES20.glBindFramebuffer(GLES20.GL_FRAMEBUFFER, previousFBO);
+				GLES32.glBindFramebuffer(GLES32.GL_FRAMEBUFFER, previousFBO);
 				if (null != destData && FBOTextureID > 0)
 				{
 					glInt[0] = FBOTextureID;
-					GLES20.glDeleteTextures(1, glInt, 0);
+					GLES32.glDeleteTextures(1, glInt, 0);
 				}
-				GLES20.glBindBuffer(GLES20.GL_ARRAY_BUFFER, previousVBO);
+				GLES32.glBindBuffer(GLES32.GL_ARRAY_BUFFER, previousVBO);
 
-				GLES20.glViewport(previousViewport[0], previousViewport[1],	previousViewport[2], previousViewport[3]);
-				if (previousBlend) GLES20.glEnable(GLES20.GL_BLEND);
-				if (previousCullFace) GLES20.glEnable(GLES20.GL_CULL_FACE);
-				if (previousScissorTest) GLES20.glEnable(GLES20.GL_SCISSOR_TEST);
-				if (previousStencilTest) GLES20.glEnable(GLES20.GL_STENCIL_TEST);
-				if (previousDepthTest) GLES20.glEnable(GLES20.GL_DEPTH_TEST);
-				if (previousDither) GLES20.glEnable(GLES20.GL_DITHER);
+				GLES32.glViewport(previousViewport[0], previousViewport[1],	previousViewport[2], previousViewport[3]);
+				if (previousBlend) GLES32.glEnable(GLES32.GL_BLEND);
+				if (previousCullFace) GLES32.glEnable(GLES32.GL_CULL_FACE);
+				if (previousScissorTest) GLES32.glEnable(GLES32.GL_SCISSOR_TEST);
+				if (previousStencilTest) GLES32.glEnable(GLES32.GL_STENCIL_TEST);
+				if (previousDepthTest) GLES32.glEnable(GLES32.GL_DEPTH_TEST);
+				if (previousDither) GLES32.glEnable(GLES32.GL_DITHER);
 
-				GLES20.glTexParameteri(GLES20.GL_TEXTURE_2D, GLES20.GL_TEXTURE_MIN_FILTER, previousMinFilter);
-				GLES20.glTexParameteri(GLES20.GL_TEXTURE_2D, GLES20.GL_TEXTURE_MAG_FILTER, previousMagFilter);
+				GLES32.glTexParameteri(GLES32.GL_TEXTURE_2D, GLES32.GL_TEXTURE_MIN_FILTER, previousMinFilter);
+				GLES32.glTexParameteri(GLES32.GL_TEXTURE_2D, GLES32.GL_TEXTURE_MAG_FILTER, previousMagFilter);
 
 				// invalidate cached state in RHI
-				GLES20.glDisableVertexAttribArray(mPositionAttrib);
-				GLES20.glDisableVertexAttribArray(mTexCoordsAttrib);
+				GLES32.glDisableVertexAttribArray(mPositionAttrib);
+				GLES32.glDisableVertexAttribArray(mTexCoordsAttrib);
 				nativeClearCachedAttributeState(mPositionAttrib, mTexCoordsAttrib); // offered by RHI for ES on Android
 			}
 
@@ -814,14 +810,14 @@ public class ElectraTextureSample
 		{
 			switch (error)
 			{
-				case GLES20.GL_INVALID_ENUM:						Log.e(TAG,"MediaPlayer$BitmapRenderer: " + op + ": glGetError GL_INVALID_ENUM");  break;
-				case GLES20.GL_INVALID_OPERATION:					Log.e(TAG,"MediaPlayer$BitmapRenderer: " + op + ": glGetError GL_INVALID_OPERATION");  break;
-				case GLES20.GL_INVALID_FRAMEBUFFER_OPERATION:		Log.e(TAG,"MediaPlayer$BitmapRenderer: " + op + ": glGetError GL_INVALID_FRAMEBUFFER_OPERATION");  break;
-				case GLES20.GL_INVALID_VALUE:						Log.e(TAG,"MediaPlayer$BitmapRenderer: " + op + ": glGetError GL_INVALID_VALUE");  break;
-				case GLES20.GL_FRAMEBUFFER_INCOMPLETE_ATTACHMENT:	Log.e(TAG,"MediaPlayer$BitmapRenderer: " + op + ": glGetError GL_FRAMEBUFFER_INCOMPLETE_ATTACHMENT");  break;
-				case GLES20.GL_FRAMEBUFFER_INCOMPLETE_DIMENSIONS:	Log.e(TAG,"MediaPlayer$BitmapRenderer: " + op + ": glGetError GL_FRAMEBUFFER_INCOMPLETE_DIMENSIONS");  break;
-				case GLES20.GL_FRAMEBUFFER_UNSUPPORTED:				Log.e(TAG,"MediaPlayer$BitmapRenderer: " + op + ": glGetError GL_FRAMEBUFFER_UNSUPPORTED");  break;
-				case GLES20.GL_OUT_OF_MEMORY:						Log.e(TAG,"MediaPlayer$BitmapRenderer: " + op + ": glGetError GL_OUT_OF_MEMORY");  break;
+				case GLES32.GL_INVALID_ENUM:						Log.e(TAG,"MediaPlayer$BitmapRenderer: " + op + ": glGetError GL_INVALID_ENUM");  break;
+				case GLES32.GL_INVALID_OPERATION:					Log.e(TAG,"MediaPlayer$BitmapRenderer: " + op + ": glGetError GL_INVALID_OPERATION");  break;
+				case GLES32.GL_INVALID_FRAMEBUFFER_OPERATION:		Log.e(TAG,"MediaPlayer$BitmapRenderer: " + op + ": glGetError GL_INVALID_FRAMEBUFFER_OPERATION");  break;
+				case GLES32.GL_INVALID_VALUE:						Log.e(TAG,"MediaPlayer$BitmapRenderer: " + op + ": glGetError GL_INVALID_VALUE");  break;
+				case GLES32.GL_FRAMEBUFFER_INCOMPLETE_ATTACHMENT:	Log.e(TAG,"MediaPlayer$BitmapRenderer: " + op + ": glGetError GL_FRAMEBUFFER_INCOMPLETE_ATTACHMENT");  break;
+				case GLES32.GL_FRAMEBUFFER_INCOMPLETE_DIMENSIONS:	Log.e(TAG,"MediaPlayer$BitmapRenderer: " + op + ": glGetError GL_FRAMEBUFFER_INCOMPLETE_DIMENSIONS");  break;
+				case GLES32.GL_FRAMEBUFFER_UNSUPPORTED:				Log.e(TAG,"MediaPlayer$BitmapRenderer: " + op + ": glGetError GL_FRAMEBUFFER_UNSUPPORTED");  break;
+				case GLES32.GL_OUT_OF_MEMORY:						Log.e(TAG,"MediaPlayer$BitmapRenderer: " + op + ": glGetError GL_OUT_OF_MEMORY");  break;
 				default:											Log.e(TAG,"MediaPlayer$BitmapRenderer: " + op + ": glGetError " + error);
 			}
 		}
@@ -829,7 +825,7 @@ public class ElectraTextureSample
 		private void glVerify(String op)
 		{
 			int error;
-			while ((error = GLES20.glGetError()) != GLES20.GL_NO_ERROR)
+			while ((error = GLES32.glGetError()) != GLES32.GL_NO_ERROR)
 			{
 				showGlError(op, error);
 				throw new RuntimeException(op + ": glGetError " + error);
@@ -839,13 +835,13 @@ public class ElectraTextureSample
 		private void glWarn(String op)
 		{
 			int error;
-			while ((error = GLES20.glGetError()) != GLES20.GL_NO_ERROR)
+			while ((error = GLES32.glGetError()) != GLES32.GL_NO_ERROR)
 			{
 				showGlError(op, error);
 			}
 		}
 
-		public FFrameUpdateInfo GetVideoFrameUpdateInfo(int destTexture, int width, int height)
+		public FFrameUpdateInfo GetVideoFrameUpdateInfo(int destTexture, int width, int height, boolean bIs10Bit)
 		{
 			mTextureWidth = width;
 			mTextureHeight = height;
@@ -853,10 +849,10 @@ public class ElectraTextureSample
 			if (mVulkanRenderer)
 			{
 				// Return data as byte buffer
-				return updateFrameData();
+				return updateFrameData(bIs10Bit);
 			}
 			// Return data in ES texture
-			return updateFrameData(destTexture);
+			return updateFrameData(destTexture, bIs10Bit);
 		}
 
 		public void release()
@@ -888,34 +884,40 @@ public class ElectraTextureSample
 			if (mBlitBuffer > 0)
 			{
 				glInt[0] = mBlitBuffer;
-				GLES20.glDeleteBuffers(1,glInt,0);
+				GLES32.glDeleteBuffers(1,glInt,0);
 				mBlitBuffer = -1;
 			}
-			if (mProgram > 0)
+			for(int i=0; i<2; ++i)
 			{
-				GLES20.glDeleteProgram(mProgram);
-				mProgram = -1;
+				if (mProgram[i] > 0)
+				{
+					GLES32.glDeleteProgram(mProgram[i]);
+					mProgram[i] = -1;
+				}
 			}
 			if (mBlitVertexShaderID > 0)
 			{
-				GLES20.glDeleteShader(mBlitVertexShaderID);
+				GLES32.glDeleteShader(mBlitVertexShaderID);
 				mBlitVertexShaderID = -1;
 			}
-			if (mBlitFragmentShaderID > 0)
+			for(int i=0; i<2; ++i)
 			{
-				GLES20.glDeleteShader(mBlitFragmentShaderID);
-				mBlitFragmentShaderID = -1;
+				if (mBlitFragmentShaderID[i] > 0)
+				{
+					GLES32.glDeleteShader(mBlitFragmentShaderID[i]);
+					mBlitFragmentShaderID[i] = -1;
+				}
 			}
 			if (mFBO > 0)
 			{
 				glInt[0] = mFBO;
-				GLES20.glDeleteFramebuffers(1,glInt,0);
+				GLES32.glDeleteFramebuffers(1,glInt,0);
 				mFBO = -1;
 			}
 			if (mTextureID > 0)
 			{
 				glInt[0] = mTextureID;
-				GLES20.glDeleteTextures(1,glInt,0);
+				GLES32.glDeleteTextures(1,glInt,0);
 				mTextureID = -1;
 			}
 			if (mEglSurface != EGL14.EGL_NO_SURFACE)
diff --git a/Engine/Plugins/Media/ElectraUtil/Source/ElectraSamples/Private/Linux/ElectraTextureSampleLinux.cpp b/Engine/Plugins/Media/ElectraUtil/Source/ElectraSamples/Private/Linux/ElectraTextureSampleLinux.cpp
--- a/Engine/Plugins/Media/ElectraUtil/Source/ElectraSamples/Private/Linux/ElectraTextureSampleLinux.cpp
+++ b/Engine/Plugins/Media/ElectraUtil/Source/ElectraSamples/Private/Linux/ElectraTextureSampleLinux.cpp
@@ -10,49 +10,11 @@ FElectraTextureSampleLinux::~FElectraTextureSampleLinux()
 {
 }
 
-
-FIntPoint FElectraTextureSampleLinux::GetDim() const
-{
-	if (VideoDecoderOutput.IsValid())
-	{
-		return VideoDecoderOutput->GetDim();
-	}
-	return FIntPoint::ZeroValue;
-}
-
-FMediaTimeStamp FElectraTextureSampleLinux::GetTime() const
-{
-	if (VideoDecoderOutput.IsValid())
-	{
-		const FDecoderTimeStamp TimeStamp = VideoDecoderOutput->GetTime();
-		return FMediaTimeStamp(TimeStamp.Time, TimeStamp.SequenceIndex);
-	}
-	return FMediaTimeStamp();
-}
-
-FTimespan FElectraTextureSampleLinux::GetDuration() const
-{
-	if (VideoDecoderOutput.IsValid())
-	{
-		return VideoDecoderOutput->GetDuration();
-	}
-	return FTimespan::Zero();
-}
-
-FIntPoint FElectraTextureSampleLinux::GetOutputDim() const
-{
-	if (VideoDecoderOutput.IsValid())
-	{
-		return VideoDecoderOutput->GetOutputDim();
-	}
-	return FIntPoint::ZeroValue;
-}
-
 uint32 FElectraTextureSampleLinux::GetStride() const
 {
 	if (VideoDecoderOutput.IsValid())
 	{
-		return VideoDecoderOutput->GetStride();
+		return VideoDecoderOutputLinux->GetStride();
 	}
 	return 0;
 }
@@ -61,33 +23,29 @@ const void* FElectraTextureSampleLinux::GetBuffer()
 {
 	if (VideoDecoderOutput)
 	{
-		return VideoDecoderOutput->GetBuffer().GetData();
+		return VideoDecoderOutputLinux->GetBuffer().GetData();
 	}
 	return nullptr;
 }
 
 EMediaTextureSampleFormat FElectraTextureSampleLinux::GetFormat() const
 {
-	int32 NumBits = VideoDecoderOutput && VideoDecoderOutput->GetIs10Bit() ? 10 : 8;
-	(void)NumBits;
-	return EMediaTextureSampleFormat::CharNV12;
+	if (VideoDecoderOutputLinux)
+	{
+		return (VideoDecoderOutputLinux->GetFormat() == PF_NV12) ? EMediaTextureSampleFormat::CharNV12 : EMediaTextureSampleFormat::P010;
+	}
+	return EMediaTextureSampleFormat::Undefined;
 }
 
 
-/**
- * Init code for realloacting an image from the the pool
- */
-void FElectraTextureSampleLinux::InitializePoolable()
-{
-}
-
 /**
  *  Return the object to the pool and inform the renderer about this...
  */
 void FElectraTextureSampleLinux::ShutdownPoolable()
 {
-	VideoDecoderOutput.Reset();
+	VideoDecoderOutputLinux = nullptr;
 	Texture = nullptr;
+	IElectraTextureSampleBase::ShutdownPoolable();
 }
 
 #endif
diff --git a/Engine/Plugins/Media/ElectraUtil/Source/ElectraSamples/Public/Android/ElectraTextureSample.h b/Engine/Plugins/Media/ElectraUtil/Source/ElectraSamples/Public/Android/ElectraTextureSample.h
--- a/Engine/Plugins/Media/ElectraUtil/Source/ElectraSamples/Public/Android/ElectraTextureSample.h
+++ b/Engine/Plugins/Media/ElectraUtil/Source/ElectraSamples/Public/Android/ElectraTextureSample.h
@@ -7,7 +7,7 @@
 #include "CoreTypes.h"
 #include "CoreGlobals.h"
 #include "Templates/SharedPointer.h"
-#include "IMediaTextureSample.h"
+#include "IElectraTextureSample.h"
 #include "Math/IntPoint.h"
 #include "Math/Range.h"
 #include "Misc/Timespan.h"
@@ -23,9 +23,9 @@ class FElectraTextureSampleSupport;
 /**
  * Texture sample generated by AndroidMedia player.
  */
-class FElectraTextureSample final
-	: public IMediaTextureSample
-	, public IMediaPoolable
+class ELECTRASAMPLES_API FElectraTextureSample final
+	: public IElectraTextureSampleBase
+	, public IMediaTextureSampleConverter
 	, public TSharedFromThis<FElectraTextureSample, ESPMode::ThreadSafe>
 {
 public:
@@ -50,9 +50,6 @@ public:
 
 	void Initialize(FVideoDecoderOutput* InVideoDecoderOutput);
 
-	virtual void InitializePoolable() override;
-	virtual void ShutdownPoolable() override;
-
 	//~ IMediaTextureSample interface
 
 	virtual const void* GetBuffer() override
@@ -60,71 +57,35 @@ public:
 		return Buffer;
 	}
 
-	virtual FIntPoint GetDim() const override
-	{
-		return VideoDecoderOutput->GetDim();
-	}
-
-	virtual FTimespan GetDuration() const override
-	{
-		return VideoDecoderOutput->GetDuration();
-	}
+	virtual EMediaTextureSampleFormat GetFormat() const override;
 
-	virtual double GetAspectRatio() const override
-	{
-		return VideoDecoderOutput->GetAspectRatio();
-	}
-
-	virtual EMediaOrientation GetOrientation() const override
-	{
-		return (EMediaOrientation)VideoDecoderOutput->GetOrientation();
-	}
-
-	virtual EMediaTextureSampleFormat GetFormat() const override
-	{
-		return EMediaTextureSampleFormat::CharBGRA;
-	}
-
-	virtual FIntPoint GetOutputDim() const override
-	{
-		return VideoDecoderOutput->GetOutputDim();
-	}
-
-	virtual uint32 GetStride() const override
-	{
-		return GetDim().X * sizeof(uint32);
-	}
+	virtual uint32 GetStride() const override;
 
 	virtual FRHITexture* GetTexture() const override
 	{
 		return Texture.GetReference();
 	}
 
-	virtual FMediaTimeStamp GetTime() const override
-	{
-		FDecoderTimeStamp Time = VideoDecoderOutput->GetTime();
-		return FMediaTimeStamp(Time.Time, Time.SequenceIndex);
-	}
-
-	virtual bool IsCacheable() const override
-	{
-		return true;
-	}
-
-	virtual bool IsOutputSrgb() const override
+	virtual IMediaTextureSampleConverter* GetMediaTextureSampleConverter() override
 	{
-		return true;
+		return this;
 	}
 
 private:
 	friend class FElectraTextureSampleSupport;
 
-	void InitializeTexture();
+	void InitializeTexture(EPixelFormat PixelFormat);
 	void SetupFromBuffer(const void* InBuffer, int32 InBufferSize);
 
+	virtual bool Convert(FTexture2DRHIRef& InDstTexture, const FConversionHints& Hints) override;
+	virtual uint32 GetConverterInfoFlags() const
+	{
+		return ConverterInfoFlags_Default;
+	}
+
 	TSharedPtr<FElectraTextureSampleSupport, ESPMode::ThreadSafe> Support;
 
-	TSharedPtr<FVideoDecoderOutputAndroid, ESPMode::ThreadSafe> VideoDecoderOutput;
+	FVideoDecoderOutputAndroid* VideoDecoderOutputAndroid;
 
 	void* Buffer;
 	int32 BufferSize;
diff --git a/Engine/Plugins/Media/ElectraUtil/Source/ElectraSamples/Public/Apple/ElectraTextureSample.h b/Engine/Plugins/Media/ElectraUtil/Source/ElectraSamples/Public/Apple/ElectraTextureSample.h
--- a/Engine/Plugins/Media/ElectraUtil/Source/ElectraSamples/Public/Apple/ElectraTextureSample.h
+++ b/Engine/Plugins/Media/ElectraUtil/Source/ElectraSamples/Public/Apple/ElectraTextureSample.h
@@ -5,9 +5,8 @@
 #include "HAL/Platform.h"
 
 #include "MediaSamples.h"
+#include "IElectraTextureSample.h"
 #include "IMediaTextureSampleConverter.h"
-#include "Misc/Timespan.h"
-#include "MediaObjectPool.h"
 #include "MediaVideoDecoderOutputApple.h"
 
 #import <AVFoundation/AVFoundation.h>
@@ -17,9 +16,8 @@ class FElectraMediaTexConvApple;
 
 
 class ELECTRASAMPLES_API FElectraTextureSample final
-	: public IMediaTextureSample
+	: public IElectraTextureSampleBase
     , public IMediaTextureSampleConverter
-	, public IMediaPoolable
 {
 public:
 	FElectraTextureSample(const TWeakPtr<FElectraMediaTexConvApple, ESPMode::ThreadSafe> & InTexConv)
@@ -42,23 +40,8 @@ public:
 		return nullptr;
 	}
 
-	virtual FIntPoint GetDim() const override;
-
-	virtual FTimespan GetDuration() const override;
-
-	virtual double GetAspectRatio() const override
-	{
-		return VideoDecoderOutput->GetAspectRatio();
-	}
-
-	virtual EMediaOrientation GetOrientation() const override
-	{
-		return (EMediaOrientation)VideoDecoderOutput->GetOrientation();
-	}
-
 	virtual EMediaTextureSampleFormat GetFormat() const override;
 
-	virtual FIntPoint GetOutputDim() const override;
 	virtual uint32 GetStride() const override;
 
 	virtual FRHITexture* GetTexture() const override
@@ -66,45 +49,20 @@ public:
 		return Texture;
 	}
 
-	virtual FMediaTimeStamp GetTime() const override;
-
-	virtual bool IsCacheable() const override
-	{
-		return true;
-	}
-
     virtual IMediaTextureSampleConverter* GetMediaTextureSampleConverter() override
     {
         return this;
     }
 
-	bool IsOutputSrgb() const override;
-	const FMatrix& GetYUVToRGBMatrix() const override;
-	bool GetFullRange() const override;
-	FMatrix44f GetSampleToRGBMatrix() const override;
-	FMatrix44f GetGamutToXYZMatrix() const override;
-	FVector2f GetWhitePoint() const override;
-	UE::Color::EEncoding GetEncodingType() const override;
-
-	virtual void InitializePoolable() override;
-	virtual void ShutdownPoolable() override;
-
 private:
 	/** The sample's texture resource. */
 	TRefCountPtr<FRHITexture2D> Texture;
 
 	/** Output data from video decoder. */
-	TSharedPtr<FVideoDecoderOutputApple, ESPMode::ThreadSafe> VideoDecoderOutput;
+	FVideoDecoderOutputApple* VideoDecoderOutputApple;
 
 	TWeakPtr<FElectraMediaTexConvApple, ESPMode::ThreadSafe> TexConv;
 
-	/** Quick access for some HDR related info */
-	TWeakPtr<const IVideoDecoderHDRInformation, ESPMode::ThreadSafe> HDRInfo;
-	TWeakPtr<const IVideoDecoderColorimetry, ESPMode::ThreadSafe> Colorimetry;
-
-	/** YUV matrix, adjusted to compensate for decoder output specific scale */
-	FMatrix44f YuvToRgbMtx;
-
 	virtual uint32 GetConverterInfoFlags() const override;
     virtual bool Convert(FTexture2DRHIRef & InDstTexture, const FConversionHints & Hints) override;
 };
diff --git a/Engine/Plugins/Media/ElectraUtil/Source/ElectraSamples/Public/IElectraTextureSample.h b/Engine/Plugins/Media/ElectraUtil/Source/ElectraSamples/Public/IElectraTextureSample.h
new file mode 100644
--- /dev/null
+++ b/Engine/Plugins/Media/ElectraUtil/Source/ElectraSamples/Public/IElectraTextureSample.h
@@ -0,0 +1,67 @@
+// Copyright Epic Games, Inc. All Rights Reserved.
+#pragma once
+
+#include "IMediaTextureSample.h"
+#include "IMediaTextureSampleConverter.h"
+#include "MediaObjectPool.h"
+#include "Misc/Timespan.h"
+#include "Templates/SharedPointer.h"
+#include "Templates/RefCounting.h"
+
+#include "MediaVideoDecoderOutput.h"
+
+class ELECTRASAMPLES_API IElectraTextureSampleBase
+	: public IMediaTextureSample
+	, public IMediaPoolable
+{
+public:
+	void Initialize(FVideoDecoderOutput* InVideoDecoderOutput);
+
+	virtual bool IsCacheable() const override
+	{
+		return true;
+	}
+
+#if !UE_SERVER
+	virtual void InitializePoolable() override;
+	virtual void ShutdownPoolable() override;
+#endif
+
+	virtual FIntPoint GetDim() const override;
+	virtual FIntPoint GetOutputDim() const override;
+
+	virtual FMediaTimeStamp GetTime() const override;
+	virtual FTimespan GetDuration() const override;
+
+	virtual double GetAspectRatio() const override
+	{
+		return VideoDecoderOutput->GetAspectRatio();
+	}
+
+	virtual EMediaOrientation GetOrientation() const override
+	{
+		return (EMediaOrientation)VideoDecoderOutput->GetOrientation();
+	}
+
+	virtual bool IsOutputSrgb() const override;
+	virtual const FMatrix& GetYUVToRGBMatrix() const override;
+	virtual bool GetFullRange() const override;
+
+	virtual FMatrix44f GetSampleToRGBMatrix() const override;
+	virtual FMatrix44f GetGamutToXYZMatrix() const override;
+	virtual FVector2f GetWhitePoint() const override;
+	virtual UE::Color::EEncoding GetEncodingType() const override;
+
+protected:
+	virtual float GetSampleDataScale(bool b10Bit) const { return 1.0f; }
+
+	/** Output data from video decoder. */
+	TSharedPtr<FVideoDecoderOutput, ESPMode::ThreadSafe> VideoDecoderOutput;
+
+	/** Quick access for some HDR related info */
+	TWeakPtr<const IVideoDecoderHDRInformation, ESPMode::ThreadSafe> HDRInfo;
+	TWeakPtr<const IVideoDecoderColorimetry, ESPMode::ThreadSafe> Colorimetry;
+
+	/** YUV matrix, adjusted to compensate for decoder output specific scale */
+	FMatrix44f YuvToRgbMtx;
+};
diff --git a/Engine/Plugins/Media/ElectraUtil/Source/ElectraSamples/Public/Linux/ElectraTextureSample.h b/Engine/Plugins/Media/ElectraUtil/Source/ElectraSamples/Public/Linux/ElectraTextureSample.h
--- a/Engine/Plugins/Media/ElectraUtil/Source/ElectraSamples/Public/Linux/ElectraTextureSample.h
+++ b/Engine/Plugins/Media/ElectraUtil/Source/ElectraSamples/Public/Linux/ElectraTextureSample.h
@@ -3,7 +3,7 @@
 #pragma once
 
 #include "CoreMinimal.h"
-#include "IMediaTextureSample.h"
+#include "IElectraTextureSample.h"
 #include "MediaObjectPool.h"
 #include "RHI.h"
 #include "RHIUtilities.h"
@@ -13,8 +13,7 @@
 class FRHITexture;
 
 class FElectraTextureSampleLinux
-	: public IMediaTextureSample
-	, public IMediaPoolable
+	: public IElectraTextureSampleBase
 {
 public:
 	FElectraTextureSampleLinux() = default;
@@ -24,83 +23,35 @@ public:
 public:
 	void Initialize(FVideoDecoderOutput* InVideoDecoderOutput)
 	{
-		VideoDecoderOutput = StaticCastSharedPtr<FVideoDecoderOutputLinux, IDecoderOutputPoolable, ESPMode::ThreadSafe>(InVideoDecoderOutput->AsShared());
+		IElectraTextureSampleBase::Initialize(InVideoDecoderOutput);
+		VideoDecoderOutputLinux = static_cast<FVideoDecoderOutputLinux*>(InVideoDecoderOutput);
 	}
-/*
-	void CreateTexture()
-	{
-		check(IsInRenderingThread());
-
-		const FRHITextureCreateDesc Desc =
-			FRHITextureCreateDesc::Create2D(TEXT("DummyTexture2D"))
-			.SetExtent(TotalSize)
-			.SetFormat(PF_B8G8R8A8)
-			.SetFlags(ETextureCreateFlags::Dynamic | ETextureCreateFlags::SRGB | ETextureCreateFlags::RenderTargetable | ETextureCreateFlags::ShaderResource)
-			.SetInitialState(ERHIAccess::SRVMask);
-
-		Texture = RHICreateTexture(Desc);
-	}
-*/
 
 	const void* GetBuffer() override;
-
-	FIntPoint GetDim() const override;
-
-	FTimespan GetDuration() const override;
-
-	double GetAspectRatio() const override
-	{
-		return VideoDecoderOutput->GetAspectRatio();
-	}
-
-	EMediaOrientation GetOrientation() const override
-	{
-		return (EMediaOrientation)VideoDecoderOutput->GetOrientation();
-	}
+	uint32 GetStride() const override;
 
 	EMediaTextureSampleFormat GetFormat() const override;
 
-	FIntPoint GetOutputDim() const override;
-	uint32 GetStride() const override;
-
 	FRHITexture* GetTexture() const override
 	{
 		return Texture;
 	}
 
-	FMediaTimeStamp GetTime() const override;
-
-	bool IsCacheable() const override
-	{
-		return true;
-	}
-
-	bool IsOutputSrgb() const override
-	{
-		return true;
-	}
-
     IMediaTextureSampleConverter* GetMediaTextureSampleConverter() override
     {
         return nullptr;
     }
 
 #if !UE_SERVER
-	void InitializePoolable() override;
 	void ShutdownPoolable() override;
 #endif
-/*
-	TRefCountPtr<FRHITexture2D> GetTextureRef() const
-	{
-		return Texture;
-	}
-*/
+
 private:
 	/** The sample's texture resource. */
 	TRefCountPtr<FRHITexture2D> Texture;
 
 	/** Output data from video decoder. */
-	TSharedPtr<FVideoDecoderOutputLinux, ESPMode::ThreadSafe> VideoDecoderOutput;
+	FVideoDecoderOutputLinux* VideoDecoderOutputLinux;
 };
 
 using FElectraTextureSamplePtr = TSharedPtr<FElectraTextureSampleLinux, ESPMode::ThreadSafe>;
diff --git a/Engine/Plugins/Media/ElectraUtil/Source/ElectraSamples/Public/Windows/ElectraTextureSample.h b/Engine/Plugins/Media/ElectraUtil/Source/ElectraSamples/Public/Windows/ElectraTextureSample.h
--- a/Engine/Plugins/Media/ElectraUtil/Source/ElectraSamples/Public/Windows/ElectraTextureSample.h
+++ b/Engine/Plugins/Media/ElectraUtil/Source/ElectraSamples/Public/Windows/ElectraTextureSample.h
@@ -1,12 +1,7 @@
 // Copyright Epic Games, Inc. All Rights Reserved.
 #pragma once
 
-#include "IMediaTextureSample.h"
-#include "IMediaTextureSampleConverter.h"
-#include "MediaObjectPool.h"
-#include "Misc/Timespan.h"
-#include "Templates/SharedPointer.h"
-#include "Templates/RefCounting.h"
+#include "IElectraTextureSample.h"
 #include "RHI.h"
 #include "RHIResources.h"
 #include "ShaderCore.h"
@@ -18,8 +13,7 @@
 #endif
 
 class ELECTRASAMPLES_API FElectraTextureSample final
-	: public IMediaTextureSample
-	, public IMediaPoolable
+	: public IElectraTextureSampleBase
 	, public IMediaTextureSampleConverter
 {
 public:
@@ -39,41 +33,11 @@ public:
 	virtual const void* GetBuffer() override;
 	virtual uint32 GetStride() const override;
 
-	virtual FIntPoint GetDim() const override;
-	virtual FIntPoint GetOutputDim() const override;
-
-	virtual FMediaTimeStamp GetTime() const override;
-	virtual FTimespan GetDuration() const override;
-
-	virtual double GetAspectRatio() const override
-	{
-		return VideoDecoderOutput->GetAspectRatio();
-	}
-
-	virtual EMediaOrientation GetOrientation() const override
-	{
-		return (EMediaOrientation)VideoDecoderOutput->GetOrientation();
-	}
-
 	virtual EMediaTextureSampleFormat GetFormat() const override
 	{
 		return SampleFormat;
 	}
 
-	virtual bool IsCacheable() const override
-	{
-		return true;
-	}
-
-	virtual bool IsOutputSrgb() const override;
-	virtual const FMatrix& GetYUVToRGBMatrix() const override;
-	virtual bool GetFullRange() const override;
-
-	virtual FMatrix44f GetSampleToRGBMatrix() const override;
-	virtual FMatrix44f GetGamutToXYZMatrix() const override;
-	virtual FVector2f GetWhitePoint() const override;
-	virtual UE::Color::EEncoding GetEncodingType() const override;
-
 #if WITH_ENGINE
 	virtual FRHITexture* GetTexture() const override
 	{
@@ -83,11 +47,6 @@ public:
 
 	IMFSample* GetMFSample();
 
-#if !UE_SERVER
-	virtual void InitializePoolable() override;
-	virtual void ShutdownPoolable() override;
-#endif
-
 	virtual IMediaTextureSampleConverter* GetMediaTextureSampleConverter() override;
 
 private:
@@ -103,15 +62,8 @@ private:
 	/** Destination Texture resource (from Rendering device) */
 	FTexture2DRHIRef Texture;
 
-	/** Output data from video decoder. */
-	TSharedPtr<FVideoDecoderOutputPC, ESPMode::ThreadSafe> VideoDecoderOutput;
-
-	/** Quick access for some HDR related info */
-	TWeakPtr<const IVideoDecoderHDRInformation, ESPMode::ThreadSafe> HDRInfo;
-	TWeakPtr<const IVideoDecoderColorimetry, ESPMode::ThreadSafe> Colorimetry;
-
-	/** YUV matrix, adjusted to compensate for decoder output specific scale */
-	FMatrix44f YuvToRgbMtx;
+	/** Output data from video decoder. Baseclass holds reference */
+	FVideoDecoderOutputPC* VideoDecoderOutputPC;
 };
 
 using FElectraTextureSamplePtr = TSharedPtr<FElectraTextureSample, ESPMode::ThreadSafe>;
diff --git a/Engine/Shaders/Private/MediaShaders.usf b/Engine/Shaders/Private/MediaShaders.usf
--- a/Engine/Shaders/Private/MediaShaders.usf
+++ b/Engine/Shaders/Private/MediaShaders.usf
@@ -130,11 +130,17 @@ void RGBConvertPS(
 	out float4 OutColor : SV_Target0)
 {
 	float2 ScaledUV = InUV * RGBConvertUB.UVScale;
-	OutColor = RGBConvertUB.Texture.Sample(RGBConvertUB.Sampler, ScaledUV);
-	if (RGBConvertUB.SrgbToLinear != 0)
+	OutColor = float4(RGBConvertUB.Texture.Sample(RGBConvertUB.Sampler, ScaledUV).xyz, 1);
+	// Inverse EOTF
+	if (RGBConvertUB.EOTF != 0)
 	{
-		OutColor.xyz = sRGBToLinear(OutColor.xyz);
+		OutColor.xyz = (RGBConvertUB.EOTF == 1) ? sRGBToLinear(OutColor.xyz) : ST2084ToLinear(OutColor.xyz);
 	}
+	// Adjust colorspace
+	float4 Col4 = float4(OutColor.xyz, 1);
+	OutColor.x = dot(RGBConvertUB.CSTransform[0], Col4);
+	OutColor.y = dot(RGBConvertUB.CSTransform[1], Col4);
+	OutColor.z = dot(RGBConvertUB.CSTransform[2], Col4);
 }
 
 
diff --git a/Engine/Source/Runtime/MediaAssets/Private/Misc/MediaTextureResource.cpp b/Engine/Source/Runtime/MediaAssets/Private/Misc/MediaTextureResource.cpp
--- a/Engine/Source/Runtime/MediaAssets/Private/Misc/MediaTextureResource.cpp
+++ b/Engine/Source/Runtime/MediaAssets/Private/Misc/MediaTextureResource.cpp
@@ -215,18 +215,27 @@ namespace MediaTextureResourceHelpers
 	{
 		switch (Sample->GetFormat())
 		{
-			// 10-bit formats
 		case EMediaTextureSampleFormat::CharBGR10A2:
+		{
+			if (Sample->GetEncodingType() == UE::Color::EEncoding::ST2084)
+			{
+				return PF_FloatRGB;
+			}
+			return PF_A2B10G10R10;
+		}
 		case EMediaTextureSampleFormat::YUVv210:
 			return PF_A2B10G10R10;
-			// Float formats
+
 		case EMediaTextureSampleFormat::FloatRGB:
-		case EMediaTextureSampleFormat::FloatRGBA:
 		case EMediaTextureSampleFormat::Y416:
 		case EMediaTextureSampleFormat::P010:
+			return PF_FloatRGB;
+
+		case EMediaTextureSampleFormat::FloatRGBA:
 		case EMediaTextureSampleFormat::P010_RGB1010102:
 			return PF_FloatRGBA;
-			// Everything else maps to 8-bit RGB...
+
+		// Everything else maps to 8-bit RGB...
 		default:
 			return PF_B8G8R8A8;
 		}
@@ -999,7 +1008,7 @@ void FMediaTextureResource::ConvertSample(const TSharedPtr<IMediaTextureSample,
 					TShaderMapRef<FRGBConvertPS> ConvertShader(ShaderMap);
 					GraphicsPSOInit.BoundShaderState.PixelShaderRHI = ConvertShader.GetPixelShader();
 					SetGraphicsPipelineState(CommandList, GraphicsPSOInit, 0);
-					ConvertShader->SetParameters(CommandList, InputTexture, OutputDim, false);
+					ConvertShader->SetParameters(CommandList, InputTexture, OutputDim, false, false, FMatrix44f::Identity);
 				}
 				break;
 
diff --git a/Engine/Source/Runtime/RenderCore/Private/MediaShaders.cpp b/Engine/Source/Runtime/RenderCore/Private/MediaShaders.cpp
--- a/Engine/Source/Runtime/RenderCore/Private/MediaShaders.cpp
+++ b/Engine/Source/Runtime/RenderCore/Private/MediaShaders.cpp
@@ -154,7 +154,7 @@ void FAYUVConvertPS::SetParameters(FRHICommandList& CommandList, TRefCountPtr<FR
 }
 
 
-/* FRGBConvertPS shader
+/* FBMPConvertPS shader
  *****************************************************************************/
 
 BEGIN_GLOBAL_SHADER_PARAMETER_STRUCT(FBMPConvertUB, )
@@ -405,23 +405,25 @@ void FP010_2101010ConvertPS::SetParameters(FRHICommandList& CommandList, const F
 
 BEGIN_GLOBAL_SHADER_PARAMETER_STRUCT(FRGBConvertUB, )
 SHADER_PARAMETER(FVector2f, UVScale)
-SHADER_PARAMETER(uint32, SrgbToLinear)
 SHADER_PARAMETER_TEXTURE(Texture2D, Texture)
 SHADER_PARAMETER_SAMPLER(SamplerState, Sampler)
+SHADER_PARAMETER(FMatrix44f, CSTransform)
+SHADER_PARAMETER(uint32, EOTF)				// 0 = linear, 1=sRGB, 2=ST2084
 END_GLOBAL_SHADER_PARAMETER_STRUCT()
 
 IMPLEMENT_GLOBAL_SHADER_PARAMETER_STRUCT(FRGBConvertUB, "RGBConvertUB");
 IMPLEMENT_SHADER_TYPE(, FRGBConvertPS, TEXT("/Engine/Private/MediaShaders.usf"), TEXT("RGBConvertPS"), SF_Pixel);
 
 
-void FRGBConvertPS::SetParameters(FRHICommandList& CommandList, TRefCountPtr<FRHITexture2D> RGBTexture, const FIntPoint& OutputDimensions, bool bSrgbToLinear)
+void FRGBConvertPS::SetParameters(FRHICommandList& CommandList, TRefCountPtr<FRHITexture2D> RGBTexture, const FIntPoint& OutputDimensions, bool bSrgbToLinear, bool bIsST2084, const FMatrix44f& CSTransform)
 {
 	FRGBConvertUB UB;
 	{
 		UB.Sampler = TStaticSamplerState<SF_Bilinear>::GetRHI();
-		UB.SrgbToLinear = bSrgbToLinear;
 		UB.Texture = RGBTexture;
 		UB.UVScale = FVector2f((float)OutputDimensions.X / (float)RGBTexture->GetSizeX(), (float)OutputDimensions.Y / (float)RGBTexture->GetSizeY());
+		UB.EOTF = bSrgbToLinear ? 1 : (bIsST2084 ? 2 : 0);
+		UB.CSTransform = CSTransform;
 	}
 
 	TUniformBufferRef<FRGBConvertUB> Data = TUniformBufferRef<FRGBConvertUB>::CreateUniformBufferImmediate(UB, UniformBuffer_SingleFrame);
diff --git a/Engine/Source/Runtime/RenderCore/Public/MediaShaders.h b/Engine/Source/Runtime/RenderCore/Public/MediaShaders.h
--- a/Engine/Source/Runtime/RenderCore/Public/MediaShaders.h
+++ b/Engine/Source/Runtime/RenderCore/Public/MediaShaders.h
@@ -333,7 +333,7 @@ public:
  * Pixel shader to resize an RGB texture.
  *
  * This shader expects an RGB or RGBA frame packed into a single texture
- * in PF_B8G8R8A8 or PF_FloatRGB format.
+ * in PF_B8G8R8A8, PF_A2B10G10R10 or PF_FloatRGB format.
  */
 class FRGBConvertPS
 	: public FGlobalShader
@@ -348,7 +348,7 @@ public:
 	{ }
 
 
-	RENDERCORE_API void SetParameters(FRHICommandList& RHICmdList, TRefCountPtr<FRHITexture2D> RGBTexture, const FIntPoint& OutputDimensions, bool SrgbToLinear);
+	RENDERCORE_API void SetParameters(FRHICommandList& RHICmdList, TRefCountPtr<FRHITexture2D> RGBTexture, const FIntPoint& OutputDimensions, bool SrgbToLinear, bool bIsST2084, const FMatrix44f& CSTransform);
 };
 
 
