Fixing up non-unity build issues in Android texture sample code (issue 248910)
#rnx

diff --git a/Engine/Plugins/Media/ElectraUtil/Source/ElectraSamples/Private/ElectraTextureSampleAndroid.cpp b/Engine/Plugins/Media/ElectraUtil/Source/ElectraSamples/Private/ElectraTextureSampleAndroid.cpp
--- a/Engine/Plugins/Media/ElectraUtil/Source/ElectraSamples/Private/ElectraTextureSampleAndroid.cpp
+++ b/Engine/Plugins/Media/ElectraUtil/Source/ElectraSamples/Private/ElectraTextureSampleAndroid.cpp
@@ -11,6 +11,8 @@
 #include "Android/AndroidApplication.h"
 
 #include "RenderingThread.h"
+#include "RHIStaticStates.h"
+#include "PipelineStateCache.h"
 
 DECLARE_GPU_STAT_NAMED(MediaAndroidDecoder_Convert, TEXT("MediaAndroidDecoder_Convert"));
 
diff --git a/Engine/Shaders/Private/GammaCorrectionCommon.ush b/Engine/Shaders/Private/GammaCorrectionCommon.ush
--- a/Engine/Shaders/Private/GammaCorrectionCommon.ush
+++ b/Engine/Shaders/Private/GammaCorrectionCommon.ush
@@ -151,7 +151,7 @@ ST2084ToLinear(float3 pq)
 	float3 Np = pow( pq, 1./m2 );
 	float3 L = Np - c1;
 	L = max(0., L);
-	L = L / (c2 - c3 * Np);
+	L = L / max(0.000000001, c2 - c3 * Np);		// DETERMINE SENSICAL CUT-OFF
 	L = pow( L, 1./m1 );
 	float3 P = L * C;
 
