Correcting slight error in luminance values for HDR output with Electra & adding clamp to secure data conversion against numerical issues that could trigger incorrectly converted pixels
#rnx

diff --git a/Engine/Plugins/Media/ElectraUtil/Source/ElectraSamples/Private/ElectraTextureSampleBase.cpp b/Engine/Plugins/Media/ElectraUtil/Source/ElectraSamples/Private/ElectraTextureSampleBase.cpp
--- a/Engine/Plugins/Media/ElectraUtil/Source/ElectraSamples/Private/ElectraTextureSampleBase.cpp
+++ b/Engine/Plugins/Media/ElectraUtil/Source/ElectraSamples/Private/ElectraTextureSampleBase.cpp
@@ -24,7 +24,9 @@ void IElectraTextureSampleBase::Initialize(FVideoDecoderOutput* InVideoDecoderOu
 	const FMatrix* Mtx = bFullRange ? &MediaShaders::YuvToRgbRec709Unscaled : &MediaShaders::YuvToRgbRec709Scaled;
 	FVector Off = (VideoDecoderOutput->GetFormat() == PF_NV12) ? (bFullRange ? MediaShaders::YUVOffsetNoScale8bits : MediaShaders::YUVOffset8bits)
 														       : (bFullRange ? MediaShaders::YUVOffsetNoScale16bits : MediaShaders::YUVOffset16bits);
-	float Scale = GetSampleDataScale(false);
+	// Correctional scale for input data
+	// (data should be placed in the upper 10-bits of the 16-bit texture channels, but some platforms do not do this)
+	float DataScale = GetSampleDataScale(false);
 
 	if (auto PinnedHDRInfo = HDRInfo.Pin())
 	{
@@ -35,21 +37,24 @@ void IElectraTextureSampleBase::Initialize(FVideoDecoderOutput* InVideoDecoderOu
 		case IVideoDecoderHDRInformation::EType::HDR10:
 			Mtx = &MediaShaders::YuvToRgbRec2020Unscaled;
 			Off = MediaShaders::YUVOffsetNoScale16bits;
-			Scale = GetSampleDataScale(true);
+			DataScale = GetSampleDataScale(true);
 			break;
 		default:
 			break;
 		}
 	}
 
+	// Compute scale to make correct towards the max value (1.0 if it is 0xff or 0xffff, but >1.0 if it 0xffc0 - which is the case for P010)
+	float NormScale = (VideoDecoderOutput->GetFormat() == PF_P010) ? (65535.0f / 65472.0f) : 1.0f;
+
 	// Matrix to transform sample data to standard YUV values
 	FMatrix PreMtx = FMatrix::Identity;
-	PreMtx.M[0][0] = Scale;
-	PreMtx.M[1][1] = Scale;
-	PreMtx.M[2][2] = Scale;
-	PreMtx.M[0][3] = -Off.X;
-	PreMtx.M[1][3] = -Off.Y;
-	PreMtx.M[2][3] = -Off.Z;
+	PreMtx.M[0][0] = DataScale * NormScale;
+	PreMtx.M[1][1] = DataScale * NormScale;
+	PreMtx.M[2][2] = DataScale * NormScale;
+	PreMtx.M[0][3] = -Off.X * NormScale;
+	PreMtx.M[1][3] = -Off.Y * NormScale;
+	PreMtx.M[2][3] = -Off.Z * NormScale;
 
 	// Combine this with the actual YUV-RGB conversion
 	YuvToRgbMtx = FMatrix44f(*Mtx * PreMtx);
diff --git a/Engine/Shaders/Private/MediaShaders.usf b/Engine/Shaders/Private/MediaShaders.usf
--- a/Engine/Shaders/Private/MediaShaders.usf
+++ b/Engine/Shaders/Private/MediaShaders.usf
@@ -357,6 +357,9 @@ void P010ConvertPS(
 	Color.y = dot(P010ConvertUB.ColorTransform[1], YUV4);
 	Color.z = dot(P010ConvertUB.ColorTransform[2], YUV4);
 
+	// Make sure we safely stay above zero
+	Color = max(Color, 0.0f);
+
 	// Go linear
 	Color = P010ConvertUB.IsST2084 ? ST2084ToLinear(Color) : Rec709ToLinear(Color);
 
@@ -399,6 +402,9 @@ void P010ConvertAsUINT16sPS(
 	Color.y = dot(P010ConvertAsUINT16sUB.ColorTransform[1], YUV4);
 	Color.z = dot(P010ConvertAsUINT16sUB.ColorTransform[2], YUV4);
 
+	// Make sure we safely stay above zero
+	Color = max(Color, 0.0f);
+
 	// Go linear
 	Color = P010ConvertAsUINT16sUB.IsST2084 ? ST2084ToLinear(Color) : Rec709ToLinear(Color);
 
@@ -444,6 +450,9 @@ void P010_2101010ConvertPS(
 	Color.y = dot(P010_2101010ConvertUB.ColorTransform[1], YUV4);
 	Color.z = dot(P010_2101010ConvertUB.ColorTransform[2], YUV4);
 
+	// Make sure we safely stay above zero
+	Color = max(Color, 0.0f);
+
 	// Go linear
 	Color = P010_2101010ConvertUB.IsST2084 ? ST2084ToLinear(Color) : Rec709ToLinear(Color);
 
