ElectraPlayer: Added metadata change event and method to access the metadata.

#rnx

diff --git a/Engine/Plugins/Media/ElectraPlayer/Source/ElectraPlayerPlugin/Private/ElectraPlayerPlugin.cpp b/Engine/Plugins/Media/ElectraPlayer/Source/ElectraPlayerPlugin/Private/ElectraPlayerPlugin.cpp
--- a/Engine/Plugins/Media/ElectraPlayer/Source/ElectraPlayerPlugin/Private/ElectraPlayerPlugin.cpp
+++ b/Engine/Plugins/Media/ElectraPlayer/Source/ElectraPlayerPlugin/Private/ElectraPlayerPlugin.cpp
@@ -33,6 +33,7 @@ FElectraPlayerPlugin::FElectraPlayerPlugin()
 	static_assert((int32)EMediaEvent::PlaybackSuspended == (int32)IElectraPlayerAdapterDelegate::EPlayerEvent::PlaybackSuspended, "check alignment of both enums");
 	static_assert((int32)EMediaEvent::SeekCompleted == (int32)IElectraPlayerAdapterDelegate::EPlayerEvent::SeekCompleted, "check alignment of both enums");
 	static_assert((int32)EMediaEvent::TracksChanged == (int32)IElectraPlayerAdapterDelegate::EPlayerEvent::TracksChanged, "check alignment of both enums");
+	static_assert((int32)EMediaEvent::MetadataChanged == (int32)IElectraPlayerAdapterDelegate::EPlayerEvent::MetadataChanged, "check alignment of both enums");
 	static_assert((int32)EMediaEvent::Internal_PurgeVideoSamplesHint == (int32)IElectraPlayerAdapterDelegate::EPlayerEvent::Internal_PurgeVideoSamplesHint, "check alignment of both enums");
 	static_assert((int32)EMediaEvent::Internal_ResetForDiscontinuity == (int32)IElectraPlayerAdapterDelegate::EPlayerEvent::Internal_ResetForDiscontinuity, "check alignment of both enums");
 	static_assert((int32)EMediaEvent::Internal_RenderClockStart == (int32)IElectraPlayerAdapterDelegate::EPlayerEvent::Internal_RenderClockStart, "check alignment of both enums");
@@ -81,6 +82,8 @@ bool FElectraPlayerPlugin::Initialize(IMediaEventSink& InEventSink,
 	PlayerDelegate = MakeShareable(new FPlayerAdapterDelegate(AsShared()));
 	Player = MakeShareable(FElectraPlayerRuntimeFactory::CreatePlayer(PlayerDelegate, InSendAnalyticMetricsDelegate, InSendAnalyticMetricsPerMinuteDelegate, InReportVideoStreamingErrorDelegate, InReportSubtitlesFileMetricsDelegate));
 
+	bMetadataChanged = false;
+	CurrentMetadata.Reset();
 	return true;
 }
 
@@ -224,6 +227,25 @@ public:
 
 //-----------------------------------------------------------------------------
 
+class FStreamMetadataItem : public IMediaPlayer::IMetadataItem
+{
+public:
+	FStreamMetadataItem(const TSharedPtr<Electra::IMediaStreamMetadata::IItem, ESPMode::ThreadSafe>& InItem) : Item(InItem.ToSharedRef())
+	{ }
+	virtual ~FStreamMetadataItem()
+	{ }
+	const FString& GetLanguageCode() const override
+	{ return Item->GetLanguageCode(); }
+	const FString& GetMimeType() const override
+	{ return Item->GetMimeType(); }
+	const FVariant& GetValue() const override
+	{ return Item->GetValue(); }
+private:
+	TSharedRef<Electra::IMediaStreamMetadata::IItem, ESPMode::ThreadSafe> Item;
+};
+
+//-----------------------------------------------------------------------------
+
 void FElectraPlayerPlugin::FPlayerAdapterDelegate::BlobReceived(const TSharedPtr<TArray<uint8>, ESPMode::ThreadSafe>& InBlobData, IElectraPlayerAdapterDelegate::EBlobResultType InResultType, int32 InResultCode, const Electra::FParamDict* InExtraInfo)
 {
 }
@@ -328,6 +350,10 @@ void FElectraPlayerPlugin::FPlayerAdapterDelegate::SendMediaEvent(EPlayerEvent E
 	TSharedPtr<FElectraPlayerPlugin, ESPMode::ThreadSafe> PinnedHost = Host.Pin();
 	if (PinnedHost.IsValid())
 	{
+		if (Event == EPlayerEvent::MetadataChanged)
+		{
+			PinnedHost->SetMetadataChanged();
+		}
 		FScopeLock lock(&PinnedHost->CallbackPointerLock);
 		if (PinnedHost->EventSink)
 		{
@@ -641,6 +667,9 @@ bool FElectraPlayerPlugin::Open(const FString& Url, const IMediaOptions* Options
 		LocalPlaystartOptions.MaxBandwidthForStreaming = (int32)MaxBandwidthForStreaming;
 	}
 
+	bMetadataChanged = false;
+	CurrentMetadata.Reset();
+
 	return Player->OpenInternal(Url, PlayerOptions, LocalPlaystartOptions, IElectraPlayerInterface::EOpenType::Media);
 }
 
@@ -679,6 +708,42 @@ void FElectraPlayerPlugin::TickInput(FTimespan DeltaTime, FTimespan Timecode)
 	Player->Tick(DeltaTime, Timecode);
 }
 
+//-----------------------------------------------------------------------------
+/**
+	Returns the current metadata, if any.
+*/
+TSharedPtr<TMap<FString, TArray<TUniquePtr<IMediaPlayer::IMetadataItem>>>, ESPMode::ThreadSafe> FElectraPlayerPlugin::GetMediaMetadata() const
+{
+	if (bMetadataChanged && Player.IsValid())
+	{
+		TSharedPtr<TMap<FString, TArray<TSharedPtr<Electra::IMediaStreamMetadata::IItem, ESPMode::ThreadSafe>>>, ESPMode::ThreadSafe> PlayerMeta = Player->GetMediaMetadata();
+		if (PlayerMeta.IsValid())
+		{
+			TSharedPtr<TMap<FString, TArray<TUniquePtr<IMediaPlayer::IMetadataItem>>>, ESPMode::ThreadSafe> NewMeta(new TMap<FString, TArray<TUniquePtr<IMediaPlayer::IMetadataItem>>>);
+			for(auto& PlayerMetaItem : *PlayerMeta)
+			{
+				TArray<TUniquePtr<IMediaPlayer::IMetadataItem>>& NewItemList = NewMeta->Emplace(PlayerMetaItem.Key);
+				for(auto& PlayerMetaListItem : PlayerMetaItem.Value)
+				{
+					if (PlayerMetaListItem.IsValid())
+					{
+						NewItemList.Emplace(MakeUnique<FStreamMetadataItem>(PlayerMetaListItem));
+					}
+				}
+			}
+			bMetadataChanged = false;
+			CurrentMetadata = MoveTemp(NewMeta);
+		}
+	}
+	return CurrentMetadata;
+}
+
+void FElectraPlayerPlugin::SetMetadataChanged()
+{
+	bMetadataChanged = true;
+}
+
+
 //-----------------------------------------------------------------------------
 /**
 	Get special feature flags states
diff --git a/Engine/Plugins/Media/ElectraPlayer/Source/ElectraPlayerPlugin/Private/ElectraPlayerPlugin.h b/Engine/Plugins/Media/ElectraPlayer/Source/ElectraPlayerPlugin/Private/ElectraPlayerPlugin.h
--- a/Engine/Plugins/Media/ElectraPlayer/Source/ElectraPlayerPlugin/Private/ElectraPlayerPlugin.h
+++ b/Engine/Plugins/Media/ElectraPlayer/Source/ElectraPlayerPlugin/Private/ElectraPlayerPlugin.h
@@ -106,11 +106,14 @@ public:
 		return false;
 	}
 
+	TSharedPtr<TMap<FString, TArray<TUniquePtr<IMediaPlayer::IMetadataItem>>>, ESPMode::ThreadSafe> GetMediaMetadata() const override;
+
 	bool GetPlayerFeatureFlag(EFeatureFlag flag) const override;
 
 	bool SetAsyncResourceReleaseNotification(IAsyncResourceReleaseNotificationRef AsyncResourceReleaseNotification) override;
 	uint32 GetNewResourcesOnOpen() const override;
 
+	void SetMetadataChanged();
 private:
 	friend class FElectraPlayerPluginModule;
 	friend class FElectraPlayerResourceDelegate;
@@ -195,6 +198,11 @@ private:
 	/** The actual player */
 	TSharedPtr<IElectraPlayerInterface, ESPMode::ThreadSafe> Player;
 
+	/** Current player stream metadata */
+	mutable bool bMetadataChanged = false;
+	mutable TSharedPtr<TMap<FString, TArray<TUniquePtr<IMediaPlayer::IMetadataItem>>>, ESPMode::ThreadSafe> CurrentMetadata;
+
+
 	/** Output sample pools */
 	TSharedPtr<FElectraTextureSamplePool, ESPMode::ThreadSafe> OutputTexturePool;
 	FElectraPlayerAudioSamplePool OutputAudioPool;
diff --git a/Engine/Plugins/Media/ElectraPlayer/Source/ElectraPlayerRuntime/Private/ElectraPlayer.cpp b/Engine/Plugins/Media/ElectraPlayer/Source/ElectraPlayerRuntime/Private/ElectraPlayer.cpp
--- a/Engine/Plugins/Media/ElectraPlayer/Source/ElectraPlayerRuntime/Private/ElectraPlayer.cpp
+++ b/Engine/Plugins/Media/ElectraPlayer/Source/ElectraPlayerRuntime/Private/ElectraPlayer.cpp
@@ -208,6 +208,7 @@ void FElectraPlayer::ClearToDefaultState()
 	bInitialSeekPerformed = false;
 	bDiscardOutputUntilCleanStart = false;
 	LastPresentedFrameDimension = FIntPoint::ZeroValue;
+	CurrentStreamMetadata.Reset();
 	CurrentlyActiveVideoStreamFormat.Reset();
 	DeferredPlayerEvents.Empty();
 	MediaUrl.Empty();
@@ -666,7 +667,7 @@ void FElectraPlayer::Tick(FTimespan DeltaTime, FTimespan Timecode)
 	if (PinnedAdapterDelegate.IsValid())
 	{
 		IElectraPlayerAdapterDelegate::EPlayerEvent Event;
-		while (DeferredEvents.Dequeue(Event))
+		while(DeferredEvents.Dequeue(Event))
 		{
 			PinnedAdapterDelegate->SendMediaEvent(Event);
 		}
@@ -1266,6 +1267,11 @@ void FElectraPlayer::GetPlaybackRange(FPlaybackRange& OutPlaybackRange) const
 	}
 }
 
+TSharedPtr<TMap<FString, TArray<TSharedPtr<Electra::IMediaStreamMetadata::IItem, ESPMode::ThreadSafe>>>, ESPMode::ThreadSafe> FElectraPlayer::GetMediaMetadata() const
+{
+	return CurrentStreamMetadata;
+}
+
 
 TSharedPtr<Electra::FTrackMetadata, ESPMode::ThreadSafe> FElectraPlayer::GetTrackStreamMetadata(EPlayerTrackType TrackType, int32 TrackIndex) const
 {
@@ -2592,14 +2598,21 @@ void FElectraPlayer::HandlePlayerEventSeekCompleted()
 	MediaStateOnSeekFinished();
 }
 
-void FElectraPlayer::HandlePlayerMediaMetadataChanged(const FString& InMetadata)
+void FElectraPlayer::HandlePlayerMediaMetadataChanged(const TSharedPtrTS<Electra::UtilsMP4::FMetadataParser>& InMetadata)
 {
-	TSharedPtr<IElectraPlayerAdapterDelegate, ESPMode::ThreadSafe> PinnedAdapterDelegate = AdapterDelegate.Pin();
-	if (PinnedAdapterDelegate.IsValid())
+	if (InMetadata.IsValid())
 	{
-		// Send out the metadata through the option query interface since we have no other means
-		// to deliver it through Media Framework at the moment.
-		/*FVariantValue Result =*/ PinnedAdapterDelegate->QueryOptions(IElectraPlayerAdapterDelegate::EOptionType::MediaMetadataUpdate, FVariantValue(InMetadata));
+		TSharedPtr<TMap<FString, TArray<TSharedPtr<Electra::IMediaStreamMetadata::IItem, ESPMode::ThreadSafe>>>, ESPMode::ThreadSafe> NewMeta = InMetadata->GetMediaStreamMetadata();
+		CurrentStreamMetadata = MoveTemp(NewMeta);		
+		DeferredEvents.Enqueue(IElectraPlayerAdapterDelegate::EPlayerEvent::MetadataChanged);
+
+		TSharedPtr<IElectraPlayerAdapterDelegate, ESPMode::ThreadSafe> PinnedAdapterDelegate = AdapterDelegate.Pin();
+		if (PinnedAdapterDelegate.IsValid())
+		{
+			// Send out the metadata through the option query interface for the time being.
+			// This should be removed in the near future.
+			/*FVariantValue Result =*/ PinnedAdapterDelegate->QueryOptions(IElectraPlayerAdapterDelegate::EOptionType::MediaMetadataUpdate, FVariantValue(InMetadata->GetAsJSON()));
+		}
 	}
 }
 
diff --git a/Engine/Plugins/Media/ElectraPlayer/Source/ElectraPlayerRuntime/Private/ElectraPlayer.h b/Engine/Plugins/Media/ElectraPlayer/Source/ElectraPlayerRuntime/Private/ElectraPlayer.h
--- a/Engine/Plugins/Media/ElectraPlayer/Source/ElectraPlayerRuntime/Private/ElectraPlayer.h
+++ b/Engine/Plugins/Media/ElectraPlayer/Source/ElectraPlayerRuntime/Private/ElectraPlayer.h
@@ -12,6 +12,7 @@
 #include "IElectraPlayerInterface.h"
 
 #include "Player/AdaptiveStreamingPlayer.h"
+#include "MediaStreamMetadata.h"
 #include "PlayerRuntimeGlobal.h"
 
 class FVideoDecoderOutput;
@@ -134,6 +135,7 @@ public:
 	int32 GetNumVideoStreams(int32 TrackIndex) const override;
 	bool GetVideoStreamFormat(FVideoStreamFormat& OutFormat, int32 InTrackIndex, int32 InStreamIndex) const override;
 	bool GetActiveVideoStreamFormat(FVideoStreamFormat& OutFormat) const override;
+	TSharedPtr<TMap<FString, TArray<TSharedPtr<Electra::IMediaStreamMetadata::IItem, ESPMode::ThreadSafe>>>, ESPMode::ThreadSafe> GetMediaMetadata() const override;
 
 	void NotifyOfOptionChange() override;
 
@@ -285,8 +287,8 @@ private:
 	};
 	struct FPlayerMetricEvent_MediaMetadataChange : public FPlayerMetricEventBase
 	{
-		FPlayerMetricEvent_MediaMetadataChange(const TSharedPtrTS<Electra::UtilsMP4::FMetadataParser>& InMetadata) : FPlayerMetricEventBase(EType::MediaMetadataChanged), NewMetadata(InMetadata->GetAsJSON()) {}
-		FString NewMetadata;
+		FPlayerMetricEvent_MediaMetadataChange(const TSharedPtrTS<Electra::UtilsMP4::FMetadataParser>& InMetadata) : FPlayerMetricEventBase(EType::MediaMetadataChanged), NewMetadata(InMetadata) {}
+		TSharedPtrTS<Electra::UtilsMP4::FMetadataParser> NewMetadata;
 	};
 	struct FPlayerMetricEvent_Error : public FPlayerMetricEventBase
 	{
@@ -410,6 +412,8 @@ private:
 	TOptional<bool>									bFrameAccurateSeeking;
 	TOptional<bool>									bEnableLooping;
 
+	TSharedPtr<TMap<FString, TArray<TSharedPtr<Electra::IMediaStreamMetadata::IItem, ESPMode::ThreadSafe>>>, ESPMode::ThreadSafe> CurrentStreamMetadata;
+
 	TOptional<FVideoStreamFormat>					CurrentlyActiveVideoStreamFormat;
 
 	FIntPoint										LastPresentedFrameDimension;
@@ -803,7 +807,7 @@ private:
 	void HandlePlayerEventJumpInPlayPosition(const FTimeValue& ToNewTime, const FTimeValue& FromTime, Metrics::ETimeJumpReason TimejumpReason);
 	void HandlePlayerEventPlaybackStopped();
 	void HandlePlayerEventSeekCompleted();
-	void HandlePlayerMediaMetadataChanged(const FString& InMetadata);
+	void HandlePlayerMediaMetadataChanged(const TSharedPtrTS<Electra::UtilsMP4::FMetadataParser>& InMetadata);
 	void HandlePlayerEventError(const FString& ErrorReason);
 	void HandlePlayerEventLogMessage(IInfoLog::ELevel InLogLevel, const FString& InLogMessage, int64 InPlayerWallclockMilliseconds);
 	void HandlePlayerEventDroppedVideoFrame();
diff --git a/Engine/Plugins/Media/ElectraPlayer/Source/ElectraPlayerRuntime/Public/IElectraPlayerInterface.h b/Engine/Plugins/Media/ElectraPlayer/Source/ElectraPlayerRuntime/Public/IElectraPlayerInterface.h
--- a/Engine/Plugins/Media/ElectraPlayer/Source/ElectraPlayerRuntime/Public/IElectraPlayerInterface.h
+++ b/Engine/Plugins/Media/ElectraPlayer/Source/ElectraPlayerRuntime/Public/IElectraPlayerInterface.h
@@ -7,6 +7,7 @@
 
 #include "ElectraPlayerMisc.h"
 #include "ElectraPlayerResourceDelegate.h"
+#include "MediaStreamMetadata.h"
 
 class FVideoDecoderOutput;
 using FVideoDecoderOutputPtr = TSharedPtr<FVideoDecoderOutput, ESPMode::ThreadSafe>;
@@ -63,6 +64,7 @@ public:
 		PlaybackSuspended,
 		SeekCompleted,
 		TracksChanged,
+		MetadataChanged,
 
 		Internal_Start,
 		Internal_PurgeVideoSamplesHint = Internal_Start,
@@ -249,6 +251,8 @@ public:
 	virtual bool GetVideoStreamFormat(FVideoStreamFormat& OutFormat, int32 InTrackIndex, int32 InStreamIndex) const = 0;
 	virtual bool GetActiveVideoStreamFormat(FVideoStreamFormat& OutFormat) const = 0;
 
+	virtual TSharedPtr<TMap<FString, TArray<TSharedPtr<Electra::IMediaStreamMetadata::IItem, ESPMode::ThreadSafe>>>, ESPMode::ThreadSafe> GetMediaMetadata() const = 0;
+
 	virtual void NotifyOfOptionChange() = 0;
 
 	// Suspends or resumes decoder instances. Not supported on all platforms.
diff --git a/Engine/Plugins/Media/ElectraUtil/Source/ElectraBase/Private/UtilsMP4.cpp b/Engine/Plugins/Media/ElectraUtil/Source/ElectraBase/Private/UtilsMP4.cpp
--- a/Engine/Plugins/Media/ElectraUtil/Source/ElectraBase/Private/UtilsMP4.cpp
+++ b/Engine/Plugins/Media/ElectraUtil/Source/ElectraBase/Private/UtilsMP4.cpp
@@ -373,7 +373,7 @@ void UtilsMP4::FMetadataParser::ParseBoxDataList(const FString& AsCategory, cons
 		RETURN_IF_READERROR(dr.Read(LanguageIndicator));
 		int32 NumRemainingDataBytes = BoxSize - 8 - (dr.GetCurrentOffset() - CurrentOffset);
 
-		FItem Item;
+		TSharedPtr<FItem, ESPMode::ThreadSafe> Item = MakeShared<FItem, ESPMode::ThreadSafe>();
 
 		// We do not handle the country indicator at the moment.
 		// Likewise, the language indicator must be a directly specified language code since we do not handle
@@ -385,9 +385,11 @@ void UtilsMP4::FMetadataParser::ParseBoxDataList(const FString& AsCategory, cons
 		}
 		if (LanguageIndicator)
 		{
-			Item.Language[0] = (char)(0x60 + ((LanguageIndicator & 0x7c00) >> 10));
-			Item.Language[1] = (char)(0x60 + ((LanguageIndicator & 0x03e0) >> 5));
-			Item.Language[2] = (char)(0x60 + (LanguageIndicator & 0x001f));
+			char Lang[3];
+			Lang[0] = (char)(0x60 + ((LanguageIndicator & 0x7c00) >> 10));
+			Lang[1] = (char)(0x60 + ((LanguageIndicator & 0x03e0) >> 5));
+			Lang[2] = (char)(0x60 + (LanguageIndicator & 0x001f));
+			Item->Language = FString(3, Lang);
 		}
 
 		bool bSet = true;
@@ -400,8 +402,8 @@ void UtilsMP4::FMetadataParser::ParseBoxDataList(const FString& AsCategory, cons
 				FString String;
 				if ((bSet = dr.ReadStringUTF8(String, NumRemainingDataBytes)) == true)
 				{
-					Item.Value = String;
-					Item.Type = 1;
+					Item->Value = String;
+					Item->Type = 1;
 				}
 				break;
 			}
@@ -412,8 +414,8 @@ void UtilsMP4::FMetadataParser::ParseBoxDataList(const FString& AsCategory, cons
 				FString String;
 				if ((bSet = dr.ReadStringUTF16(String, NumRemainingDataBytes)) == true)
 				{
-					Item.Value = String;
-					Item.Type = 1;
+					Item->Value = String;
+					Item->Type = 1;
 				}
 				break;
 			}
@@ -423,8 +425,20 @@ void UtilsMP4::FMetadataParser::ParseBoxDataList(const FString& AsCategory, cons
 			case 27:	// BMP image
 			{
 				TArray<uint8> Image(dr.GetCurrentDataPointer(), NumRemainingDataBytes);
-				Item.Value = MoveTemp(Image);
-				Item.Type = WKT;
+				Item->Value = MoveTemp(Image);
+				Item->Type = WKT;
+				if (WKT == 13)
+				{
+					Item->MimeType = TEXT("image/jpeg");
+				}
+				else if (WKT == 14)
+				{
+					Item->MimeType = TEXT("image/png");
+				}
+				else if (WKT == 27)
+				{
+					Item->MimeType = TEXT("image/bmp");
+				}
 				break;
 			}
 
@@ -435,8 +449,8 @@ void UtilsMP4::FMetadataParser::ParseBoxDataList(const FString& AsCategory, cons
 				if ((bSet = dr.ReadAsNumber(Flt)) == true)
 				{
 					// Set as double!
-					Item.Value = (double)Flt;
-					Item.Type = 24;
+					Item->Value = (double)Flt;
+					Item->Type = 24;
 				}
 				break;
 			}
@@ -446,8 +460,8 @@ void UtilsMP4::FMetadataParser::ParseBoxDataList(const FString& AsCategory, cons
 				double Dbl;
 				if ((bSet = dr.ReadAsNumber(Dbl)) == true)
 				{
-					Item.Value = Dbl;
-					Item.Type = 24;
+					Item->Value = Dbl;
+					Item->Type = 24;
 				}
 				break;
 			}
@@ -461,8 +475,8 @@ void UtilsMP4::FMetadataParser::ParseBoxDataList(const FString& AsCategory, cons
 				int64 Number;
 				if ((bSet = dr.ReadAsNumber(Number, WKT==65?1:WKT==66?2:WKT==67?4:WKT==74?8:NumRemainingDataBytes)) == true)
 				{
-					Item.Value = Number;
-					Item.Type = 74;
+					Item->Value = Number;
+					Item->Type = 74;
 				}
 				break;
 			}
@@ -476,8 +490,8 @@ void UtilsMP4::FMetadataParser::ParseBoxDataList(const FString& AsCategory, cons
 				uint64 Number;
 				if ((bSet = dr.ReadAsNumber(Number, WKT==75?1:WKT==76?2:WKT==77?4:WKT==78?8:NumRemainingDataBytes)) == true)
 				{
-					Item.Value = Number;
-					Item.Type = 78;
+					Item->Value = Number;
+					Item->Type = 78;
 				}
 				break;
 			}
@@ -559,7 +573,7 @@ bool UtilsMP4::FMetadataParser::IsDifferentFrom(const UtilsMP4::FMetadataParser&
 	// Check if the maps keys or items are different.
 	for(auto& Key : Items)
 	{
-		const TArray<FItem>* OtherKeyItems = Other.Items.Find(Key.Key);
+		const TArray<TSharedPtr<FItem, ESPMode::ThreadSafe>>* OtherKeyItems = Other.Items.Find(Key.Key);
 		if (!OtherKeyItems || Key.Value.Num() != OtherKeyItems->Num())
 		{
 			return true;
@@ -567,7 +581,10 @@ bool UtilsMP4::FMetadataParser::IsDifferentFrom(const UtilsMP4::FMetadataParser&
 		// Expensive item by item test :-(
 		for(auto& Item : Key.Value)
 		{
-			if (!OtherKeyItems->Contains(Item))
+			if (!OtherKeyItems->ContainsByPredicate([Item](const TSharedPtr<FItem, ESPMode::ThreadSafe>& CompItem)
+			{
+				return *Item == *CompItem;
+			}))
 			{
 				return true;
 			}
@@ -577,6 +594,26 @@ bool UtilsMP4::FMetadataParser::IsDifferentFrom(const UtilsMP4::FMetadataParser&
 }
 
 
+TSharedPtr<TMap<FString, TArray<TSharedPtr<IMediaStreamMetadata::IItem, ESPMode::ThreadSafe>>>, ESPMode::ThreadSafe> UtilsMP4::FMetadataParser::GetMediaStreamMetadata() const
+{
+	TSharedPtr<TMap<FString, TArray<TSharedPtr<IMediaStreamMetadata::IItem, ESPMode::ThreadSafe>>>, ESPMode::ThreadSafe> NewMeta(new TMap<FString, TArray<TSharedPtr<IMediaStreamMetadata::IItem, ESPMode::ThreadSafe>>>);
+	for(auto& Item : Items)
+	{
+		const TArray<TSharedPtr<FItem, ESPMode::ThreadSafe>>& SrcList = Item.Value;
+		if (SrcList.Num())
+		{
+			TArray<TSharedPtr<IMediaStreamMetadata::IItem, ESPMode::ThreadSafe>> DstList;
+			for(int32 i=0; i<SrcList.Num(); ++i)
+			{
+				DstList.Emplace(SrcList[i]);
+			}
+			NewMeta->Emplace(Item.Key, MoveTemp(DstList));
+		}
+	}
+	return NewMeta;
+}
+
+
 FString UtilsMP4::FMetadataParser::GetAsJSON() const
 {
 	FString JSON(TEXT("{"));
@@ -584,12 +621,12 @@ FString UtilsMP4::FMetadataParser::GetAsJSON() const
 	for(auto& Item : Items)
 	{
 		bool bFirstItemValue = true;
-		const TArray<FItem>& List=Item.Value;
+		const TArray<TSharedPtr<FItem, ESPMode::ThreadSafe>>& List=Item.Value;
 		FString ItemJSON;
 		for(int32 i=0; i<List.Num(); ++i)
 		{
-			const FItem& it = List[i];
-			FString ItemValue = it.ToJSONValue();
+			const TSharedPtr<FItem, ESPMode::ThreadSafe>& it = List[i];
+			FString ItemValue = it->ToJSONValue();
 			if (ItemValue.Len())
 			{
 				if (!bFirstItemValue)
@@ -597,8 +634,7 @@ FString UtilsMP4::FMetadataParser::GetAsJSON() const
 					ItemJSON.AppendChar(TCHAR(','));
 				}
 				bFirstItemValue = false;
-				FString Lang(3, it.Language);
-				ItemJSON.Append(FString::Printf(TEXT("\"%s\":"), *Lang));
+				ItemJSON.Append(FString::Printf(TEXT("\"%s\":"), *(it->Language)));
 				ItemJSON.Append(MoveTemp(ItemValue));
 			}
 		}
diff --git a/Engine/Plugins/Media/ElectraUtil/Source/ElectraBase/Public/MediaStreamMetadata.h b/Engine/Plugins/Media/ElectraUtil/Source/ElectraBase/Public/MediaStreamMetadata.h
new file mode 100644
--- /dev/null
+++ b/Engine/Plugins/Media/ElectraUtil/Source/ElectraBase/Public/MediaStreamMetadata.h
@@ -0,0 +1,33 @@
+// Copyright Epic Games, Inc. All Rights Reserved.
+
+#pragma once
+
+#include "CoreMinimal.h"
+#include "Templates/SharedPointer.h"
+#include "Containers/Array.h"
+#include "Containers/Map.h"
+#include "Misc/ScopeLock.h"
+#include "Misc/Variant.h"
+
+namespace Electra
+{
+
+class ELECTRABASE_API IMediaStreamMetadata
+{
+public:
+	IMediaStreamMetadata() = default;
+	virtual ~IMediaStreamMetadata() = default;
+	IMediaStreamMetadata(const IMediaStreamMetadata& rhs) = delete;
+	IMediaStreamMetadata& operator=(const IMediaStreamMetadata&rhs) = delete;
+
+	class IItem
+	{
+	public:
+		virtual ~IItem() = default;
+		virtual const FString& GetLanguageCode() const = 0;
+		virtual const FString& GetMimeType() const = 0;
+		virtual const FVariant& GetValue() const = 0;
+	};
+};
+
+} // namespace Electra
diff --git a/Engine/Plugins/Media/ElectraUtil/Source/ElectraBase/Public/Utilities/UtilsMP4.h b/Engine/Plugins/Media/ElectraUtil/Source/ElectraBase/Public/Utilities/UtilsMP4.h
--- a/Engine/Plugins/Media/ElectraUtil/Source/ElectraBase/Public/Utilities/UtilsMP4.h
+++ b/Engine/Plugins/Media/ElectraUtil/Source/ElectraBase/Public/Utilities/UtilsMP4.h
@@ -4,6 +4,7 @@
 
 #include <CoreMinimal.h>
 #include <Misc/Variant.h>
+#include "MediaStreamMetadata.h"
 
 namespace Electra
 {
@@ -115,24 +116,33 @@ public:
 		EResult Parse(uint32 InHandler, uint32 InHandlerReserved0, const TArray<FBoxInfo>& InBoxes);
 		bool IsDifferentFrom(const FMetadataParser& Other);
 		FString GetAsJSON() const;
+		TSharedPtr<TMap<FString, TArray<TSharedPtr<IMediaStreamMetadata::IItem, ESPMode::ThreadSafe>>>, ESPMode::ThreadSafe> GetMediaStreamMetadata() const;
 	private:
 		FString PrintableBoxAtom(const uint32 InAtom);
 		void Parse(const FBoxInfo& InBox);
 		void ParseBoxDataList(const FString& AsCategory, const uint8* InBoxData, uint32 InBoxSize);
 		void ParseBoxDataiTunes(const uint8* InBoxData, uint32 InBoxSize);
 
-		struct FItem
+		class FItem : public IMediaStreamMetadata::IItem
 		{
-			char Language[3] { 0 };			// ISO 639-2; if not set (all zero) the default entry for all languages
+		public:
+			~FItem() {}
+			const FString& GetLanguageCode() const override
+			{ return Language; }
+			const FString& GetMimeType() const override
+			{ return MimeType; }
+			const FVariant& GetValue() const override
+			{ return Value; }
+
+			FString Language;				// ISO 639-2; if not set (all zero) the default entry for all languages
+			FString MimeType;
 			int32 Type = 0;					// Well-known data type (see Quicktime reference)
 			FVariant Value;
 			FString ToJSONValue() const;
 			static TArray<TCHAR> CharsToEscapeInJSON;
 			bool operator != (const FItem& Other) const
 			{
-				return Type != Other.Type ||
-						FMemory::Memcmp(Language, Other.Language, sizeof(Language)) != 0 ||
-						Value != Other.Value;
+				return Type != Other.Type || Language != Other.Language || Value != Other.Value;
 			}
 			bool operator == (const FItem& Other) const
 			{
@@ -141,7 +151,7 @@ public:
 		};
 
 		TMap<uint32, FString> WellKnownItems;
-		TMap<FString, TArray<FItem>> Items;
+		TMap<FString, TArray<TSharedPtr<FItem, ESPMode::ThreadSafe>>> Items;
 		uint32 NumTotalItems = 0;
 	};
 
diff --git a/Engine/Source/Runtime/Media/Public/IMediaEventSink.h b/Engine/Source/Runtime/Media/Public/IMediaEventSink.h
--- a/Engine/Source/Runtime/Media/Public/IMediaEventSink.h
+++ b/Engine/Source/Runtime/Media/Public/IMediaEventSink.h
@@ -38,6 +38,9 @@ enum class EMediaEvent
 	/** Media tracks have changed. */
 	TracksChanged,
 
+	/** Media metadata has changed. */
+	MetadataChanged,
+
 	// - - - - - - - - - - - - - - - - - - - - - - - -
 
 	/** Start of internal events not sent out to user code. */
diff --git a/Engine/Source/Runtime/Media/Public/IMediaPlayer.h b/Engine/Source/Runtime/Media/Public/IMediaPlayer.h
--- a/Engine/Source/Runtime/Media/Public/IMediaPlayer.h
+++ b/Engine/Source/Runtime/Media/Public/IMediaPlayer.h
@@ -3,6 +3,8 @@
 #pragma once
 
 #include "Containers/UnrealString.h"
+#include "Containers/Array.h"
+#include "Containers/Map.h"
 #include "Internationalization/Text.h"
 #include "Math/Interval.h"
 #include "Math/IntPoint.h"
@@ -194,6 +196,52 @@ public:
 		return FVariant();
 	}
 
+	/**
+	 * This structure describes a metadata item.
+	 */
+	class IMetadataItem
+	{
+	public:
+		virtual ~IMetadataItem() = default;
+
+		/*
+		  Language this item is described in.
+		  This should a three letter ISO 639-2 code. If no language is specified this item is intended
+		  to be used when no other item exists in the user's preferred language.
+		*/
+		virtual const FString& GetLanguageCode() const = 0;
+		
+		/*
+		  Optional mime type of the item described by Value.
+		  This should be set when Value is of type EVariantTypes::ByteArray
+		  For example, if this item is a binary JPEG image the mime type should be set as "image/jpeg"
+		  to identify the type of the payload.
+		*/
+		virtual const FString& GetMimeType() const = 0;
+
+		/*
+		  The value of the metadata item.
+		*/
+		virtual const FVariant& GetValue() const = 0;
+	};
+
+	/**
+	 * Gets the current metadata of the media source.
+	 * 
+	 * Metadata is optional and if present is typically a collection of key/value items without a
+	 * well defined meaning. It may contain information on copyright, album name, artist and such,
+	 * but the availability of any item is not mandatory and the representation will vary with the
+	 * type of media.
+	 * Interpretation therefor requires the application to be aware of the type of media being loaded
+	 * and the possible metadata it may carry.
+	 * 
+	 * Metadata may change over time. Its presence or change is reported by a MetadataChanged event.
+	 */
+	virtual TSharedPtr<TMap<FString, TArray<TUniquePtr<IMetadataItem>>>, ESPMode::ThreadSafe> GetMediaMetadata() const
+	{
+		return nullptr;
+	}
+
 	/**
 	 * Get the human readable name of the currently loaded media source.
 	 *
diff --git a/Engine/Source/Runtime/MediaAssets/Private/Assets/MediaPlayer.cpp b/Engine/Source/Runtime/MediaAssets/Private/Assets/MediaPlayer.cpp
--- a/Engine/Source/Runtime/MediaAssets/Private/Assets/MediaPlayer.cpp
+++ b/Engine/Source/Runtime/MediaAssets/Private/Assets/MediaPlayer.cpp
@@ -1016,6 +1016,10 @@ void UMediaPlayer::HandlePlayerMediaEvent(EMediaEvent Event)
 	case EMediaEvent::TracksChanged:
 		OnTracksChanged.Broadcast();
 		break;
+
+	case EMediaEvent::MetadataChanged:
+		OnMetadataChanged.Broadcast();
+		break;
 	}
 }
 
diff --git a/Engine/Source/Runtime/MediaAssets/Public/MediaPlayer.h b/Engine/Source/Runtime/MediaAssets/Public/MediaPlayer.h
--- a/Engine/Source/Runtime/MediaAssets/Public/MediaPlayer.h
+++ b/Engine/Source/Runtime/MediaAssets/Public/MediaPlayer.h
@@ -953,6 +953,13 @@ public:
 	UPROPERTY(BlueprintAssignable, Category="Media|MediaPlayer", meta = (HideInDetailPanel))
 	FOnMediaPlayerMediaEvent OnTracksChanged;
 
+	/**
+	 * A delegate that is invoked when the media metadata changed.
+	 *
+	 * @see OnMediaOpened
+	 */
+	UPROPERTY(BlueprintAssignable, Category="Media|MediaPlayer", meta = (HideInDetailPanel))
+	FOnMediaPlayerMediaEvent OnMetadataChanged;
 public:
 
 	/**
diff --git a/Engine/Source/Runtime/MediaUtils/Private/MediaHelpers.cpp b/Engine/Source/Runtime/MediaUtils/Private/MediaHelpers.cpp
--- a/Engine/Source/Runtime/MediaUtils/Private/MediaHelpers.cpp
+++ b/Engine/Source/Runtime/MediaUtils/Private/MediaHelpers.cpp
@@ -23,6 +23,7 @@ namespace MediaUtils
 		case EMediaEvent::PlaybackSuspended: return TEXT("PlaybackSuspended");
 		case EMediaEvent::SeekCompleted: return TEXT("SeekCompleted");
 		case EMediaEvent::TracksChanged: return TEXT("TracksChanged");
+		case EMediaEvent::MetadataChanged: return TEXT("MetadataChanged");
 
 		case EMediaEvent::Internal_PurgeVideoSamplesHint: return TEXT("Internal: PurgeVideoSamplesHint");
 		case EMediaEvent::Internal_ResetForDiscontinuity: return TEXT("Internal: ResetForDiscontinuity");
