Electra D3D12 support

diff --git a/Engine/Plugins/Experimental/TextureMediaPlayer/Source/TextureMediaPlayer/Private/TextureMediaPlayerVideoDecoderOutput.cpp b/Engine/Plugins/Experimental/TextureMediaPlayer/Source/TextureMediaPlayer/Private/TextureMediaPlayerVideoDecoderOutput.cpp
--- a/Engine/Plugins/Experimental/TextureMediaPlayer/Source/TextureMediaPlayer/Private/TextureMediaPlayerVideoDecoderOutput.cpp
+++ b/Engine/Plugins/Experimental/TextureMediaPlayer/Source/TextureMediaPlayer/Private/TextureMediaPlayerVideoDecoderOutput.cpp
@@ -60,16 +60,12 @@ TRefCountPtr<IMFSample> TextureMediaPlayerVideoDecoderOutput::GetMFSample() cons
 	return nullptr;
 }
 
-// TODO: Use commented lines to replace following lines once Electra changes are integrated
-//TRefCountPtr<IUnknown> TextureMediaPlayerVideoDecoderOutput::GetTexture() const
-TRefCountPtr<ID3D11Texture2D> TextureMediaPlayerVideoDecoderOutput::GetTexture() const
+TRefCountPtr<IUnknown> TextureMediaPlayerVideoDecoderOutput::GetTexture() const
 {
-	//IUnknown* Result = nullptr;
-	ID3D11Texture2D* Result = nullptr;
+	IUnknown* Result = nullptr;
 	if (Texture.IsValid())
 	{
-		//Result = static_cast<IUnknown*>(Texture->GetNativeResource());
-		Result = static_cast<ID3D11Texture2D*>(Texture->GetNativeResource());
+		Result = static_cast<IUnknown*>(Texture->GetNativeResource());
 	}
 	return Result;
 }
diff --git a/Engine/Plugins/Experimental/TextureMediaPlayer/Source/TextureMediaPlayer/Private/TextureMediaPlayerVideoDecoderOutput.h b/Engine/Plugins/Experimental/TextureMediaPlayer/Source/TextureMediaPlayer/Private/TextureMediaPlayerVideoDecoderOutput.h
--- a/Engine/Plugins/Experimental/TextureMediaPlayer/Source/TextureMediaPlayer/Private/TextureMediaPlayerVideoDecoderOutput.h
+++ b/Engine/Plugins/Experimental/TextureMediaPlayer/Source/TextureMediaPlayer/Private/TextureMediaPlayerVideoDecoderOutput.h
@@ -36,11 +36,8 @@ public:
 #if PLATFORM_WINDOWS
 	virtual TRefCountPtr<IMFSample> GetMFSample() const override;
 
-	// TODO: Use commented line to replace following line once Electra changes are integrated
-	//virtual TRefCountPtr<IUnknown> GetTexture() const override;
-	virtual TRefCountPtr<ID3D11Texture2D> GetTexture() const override;
-	// TODO: Re-add override specifier once Electra changes are integrated
-	virtual TRefCountPtr<IUnknown> GetSync(uint64& SyncValue) const;
+	virtual TRefCountPtr<IUnknown> GetTexture() const override;
+	virtual TRefCountPtr<IUnknown> GetSync(uint64& SyncValue) const override;
 	virtual TRefCountPtr<ID3D11Device> GetDevice() const override;
 #endif // PLATFORM_WINDOWS
 
diff --git a/Engine/Plugins/Media/ElectraPlayer/Source/ElectraPlayerRuntime/Private/Runtime/Decoder/Windows/ElectraVideoDecoder_PC.cpp b/Engine/Plugins/Media/ElectraPlayer/Source/ElectraPlayerRuntime/Private/Runtime/Decoder/Windows/ElectraVideoDecoder_PC.cpp
--- a/Engine/Plugins/Media/ElectraPlayer/Source/ElectraPlayerRuntime/Private/Runtime/Decoder/Windows/ElectraVideoDecoder_PC.cpp
+++ b/Engine/Plugins/Media/ElectraPlayer/Source/ElectraPlayerRuntime/Private/Runtime/Decoder/Windows/ElectraVideoDecoder_PC.cpp
@@ -44,9 +44,9 @@ uint32 FElectraPlayerVideoDecoderOutputPC::GetStride() const
 	return Stride;
 }
 
-TRefCountPtr<ID3D11Texture2D> FElectraPlayerVideoDecoderOutputPC::GetTexture() const
+TRefCountPtr<IUnknown> FElectraPlayerVideoDecoderOutputPC::GetTexture() const
 {
-	return Texture;
+	return static_cast<IUnknown*>(Texture.GetReference());
 }
 
 TRefCountPtr<ID3D11Device> FElectraPlayerVideoDecoderOutputPC::GetDevice() const
@@ -60,6 +60,13 @@ FIntPoint FElectraPlayerVideoDecoderOutputPC::GetDim() const
 	return SampleDim;
 }
 
+TRefCountPtr<IUnknown> FElectraPlayerVideoDecoderOutputPC::GetSync(uint64& SyncValue) const
+{
+	TRefCountPtr<IUnknown> KeyedMutex;
+	HRESULT res = Texture->QueryInterface(_uuidof(IDXGIKeyedMutex), (void**)&KeyedMutex);
+	return KeyedMutex;
+}
+
 bool FElectraPlayerVideoDecoderOutputPC::PreInitForDecode(FIntPoint OutputDim, const TFunction<void(int32 /*ApiReturnValue*/, const FString& /*Message*/, uint16 /*Code*/, UEMediaError /*Error*/ )>& PostError)
 {
 	DWORD SwBufAllocSize = 0;
diff --git a/Engine/Plugins/Media/ElectraPlayer/Source/ElectraPlayerRuntime/Private/Runtime/Decoder/Windows/ElectraVideoDecoder_PC.h b/Engine/Plugins/Media/ElectraPlayer/Source/ElectraPlayerRuntime/Private/Runtime/Decoder/Windows/ElectraVideoDecoder_PC.h
--- a/Engine/Plugins/Media/ElectraPlayer/Source/ElectraPlayerRuntime/Private/Runtime/Decoder/Windows/ElectraVideoDecoder_PC.h
+++ b/Engine/Plugins/Media/ElectraPlayer/Source/ElectraPlayerRuntime/Private/Runtime/Decoder/Windows/ElectraVideoDecoder_PC.h
@@ -57,12 +57,14 @@ public:
 
 	virtual uint32 GetStride() const override;
 
-	virtual TRefCountPtr<ID3D11Texture2D> GetTexture() const override;
+	virtual TRefCountPtr<IUnknown> GetTexture() const override;
 
 	virtual TRefCountPtr<ID3D11Device> GetDevice() const override;
 
 	virtual FIntPoint GetDim() const override;
 
+	virtual TRefCountPtr<IUnknown> GetSync(uint64& SyncValue) const override;
+
 private:
 	// Decoder output type
 	EOutputType OutputType;
diff --git a/Engine/Plugins/Media/ElectraUtil/Source/ElectraBase/Public/MediaVideoDecoderOutputPC.h b/Engine/Plugins/Media/ElectraUtil/Source/ElectraBase/Public/MediaVideoDecoderOutputPC.h
--- a/Engine/Plugins/Media/ElectraUtil/Source/ElectraBase/Public/MediaVideoDecoderOutputPC.h
+++ b/Engine/Plugins/Media/ElectraUtil/Source/ElectraBase/Public/MediaVideoDecoderOutputPC.h
@@ -9,41 +9,47 @@
 #include "Containers/Array.h"
 #include "PixelFormat.h"
 
-#include "Windows/AllowWindowsPlatformTypes.h"
+#if PLATFORM_WINDOWS
 
-#pragma warning(push)
-#pragma warning(disable : 4005)	// macro redefinition
-# include <d3d11.h>
-#pragma warning(pop)
+	#include "Windows/AllowWindowsPlatformTypes.h"
+
+	#pragma warning(push)
+	#pragma warning(disable : 4005) // macro redefinition
+	#include <d3d11.h>
+	#pragma warning(pop)
 
 THIRD_PARTY_INCLUDES_START
-#include "mfobjects.h"
-#include "mfapi.h"
+	#include "mfobjects.h"
+	#include "mfapi.h"
 THIRD_PARTY_INCLUDES_END
 
+#endif // PLATFORM_WINDOWS
+
 class FVideoDecoderOutputPC : public FVideoDecoderOutput
 {
 public:
 	enum class EOutputType
 	{
 		Unknown = 0,
-		SoftwareWin7,			// SW decode into CPU buffer
-		SoftwareWin8Plus,		// SW decode into DX11 texture
-		HardwareWin8Plus,		// HW decode into shared DX11 texture
-		HardwareDX9_DX12,		// HW decode into CPU buffer
-		Hardware_DX,			// HW decode into DX texture on render device (DX11 or DX12)
+		SoftwareWin7,	  // SW decode into CPU buffer
+		SoftwareWin8Plus, // SW decode into DX11 texture
+		HardwareWin8Plus, // HW decode into shared DX11 texture
+		HardwareDX9_DX12, // HW decode into CPU buffer
+		Hardware_DX,	  // HW decode into DX texture on render device (DX11 or DX12)
 	};
 
 	virtual EOutputType GetOutputType() const = 0;
-
-	virtual TRefCountPtr<IMFSample> GetMFSample() const = 0;
-
-	virtual const TArray<uint8> & GetBuffer() const = 0;
+	virtual const TArray<uint8>& GetBuffer() const = 0;
 	virtual uint32 GetStride() const = 0;
 
-	virtual TRefCountPtr<ID3D11Texture2D> GetTexture() const = 0;
+#if PLATFORM_WINDOWS
+	virtual TRefCountPtr<IMFSample> GetMFSample() const = 0;
+	virtual TRefCountPtr<IUnknown> GetTexture() const = 0;
+	virtual TRefCountPtr<IUnknown> GetSync(uint64& SyncValue) const = 0;
 	virtual TRefCountPtr<ID3D11Device> GetDevice() const = 0;
+#endif // PLATFORM_WINDOWS
 };
 
-#include "Windows/HideWindowsPlatformTypes.h"
-
+#if PLATFORM_WINDOWS
+	#include "Windows/HideWindowsPlatformTypes.h"
+#endif // PLATFORM_WINDOWS
diff --git a/Engine/Plugins/Media/ElectraUtil/Source/ElectraSamples/Private/ElectraTextureSamplePC.cpp b/Engine/Plugins/Media/ElectraUtil/Source/ElectraSamples/Private/ElectraTextureSamplePC.cpp
--- a/Engine/Plugins/Media/ElectraUtil/Source/ElectraSamples/Private/ElectraTextureSamplePC.cpp
+++ b/Engine/Plugins/Media/ElectraUtil/Source/ElectraSamples/Private/ElectraTextureSamplePC.cpp
@@ -9,12 +9,13 @@
 #include "ProfilingDebugging/RealtimeGPUProfiler.h"
 #include "RenderUtils.h"
 
-#ifdef ELECTRA_HAVE_DX11
 #include "Windows/AllowWindowsPlatformTypes.h"
+#ifdef ELECTRA_HAVE_DX11
 #include "D3D11State.h"
 #include "D3D11Resources.h"
-#include "Windows/HideWindowsPlatformTypes.h"
 #endif
+#include "d3d12.h"
+#include "Windows/HideWindowsPlatformTypes.h"
 
 /*
 	Short summary of how we get data:
@@ -42,26 +43,101 @@ void FElectraTextureSample::Initialize(FVideoDecoderOutput *InVideoDecoderOutput
 	IElectraTextureSampleBase::Initialize(InVideoDecoderOutput);
 	VideoDecoderOutputPC = static_cast<FVideoDecoderOutputPC*>(InVideoDecoderOutput);
 	SampleFormat = (VideoDecoderOutput->GetFormat() == PF_NV12) ? EMediaTextureSampleFormat::CharNV12 : EMediaTextureSampleFormat::P010;
+
+	ID3D12Device* ApplicationDxDevice = static_cast<ID3D12Device*>(GDynamicRHI->RHIGetNativeDevice());
+	HRESULT Res;
+	if (D3DCmdAllocator.IsValid() == false)
+	{
+		Res = ApplicationDxDevice->CreateCommandAllocator(D3D12_COMMAND_LIST_TYPE_DIRECT, IID_PPV_ARGS(D3DCmdAllocator.GetInitReference()));
+		check(!FAILED(Res));
+	}
+	if (D3DCmdList.IsValid() == false)
+	{
+		Res = ApplicationDxDevice->CreateCommandList(0, D3D12_COMMAND_LIST_TYPE_DIRECT, D3DCmdAllocator.GetReference(), nullptr, __uuidof(ID3D12CommandList), reinterpret_cast<void**>(D3DCmdList.GetInitReference()));
+		check(!FAILED(Res));
+		D3DCmdList->Close();
+	}
 }
 
 IMediaTextureSampleConverter* FElectraTextureSample::GetMediaTextureSampleConverter()
 {
-#ifdef ELECTRA_HAVE_DX11
 	// All versions might need SW fallback - check if we have a real texture as source -> converter needed
 	return (VideoDecoderOutputPC && VideoDecoderOutputPC->GetTexture()) ? this : nullptr;
-#else
-	return nullptr;
-#endif
 }
 
+
+struct FRHICommandCopyResourceDX12 final : public FRHICommand<FRHICommandCopyResourceDX12>
+{
+	TRefCountPtr<ID3D12Resource> SampleTexture;
+	FTexture2DRHIRef SampleDestinationTexture;
+	TRefCountPtr<ID3D12CommandAllocator> D3DCmdAllocator;
+	TRefCountPtr<ID3D12GraphicsCommandList> D3DCmdList;
+	TRefCountPtr<ID3D12Fence> D3DFence;
+	uint64 FenceValue;
+
+	FRHICommandCopyResourceDX12(ID3D12Resource* InSampleTexture, FRHITexture2D* InSampleDestinationTexture, TRefCountPtr<ID3D12CommandAllocator> InD3DCmdAllocator, TRefCountPtr<ID3D12GraphicsCommandList> InD3DCmdList, TRefCountPtr<ID3D12Fence> InD3DFence, uint64 InFenceValue)
+		: SampleTexture(InSampleTexture)
+		, SampleDestinationTexture(InSampleDestinationTexture)
+		, D3DCmdAllocator(InD3DCmdAllocator)
+		, D3DCmdList(InD3DCmdList)
+		, D3DFence(InD3DFence)
+		, FenceValue(InFenceValue)
+	{
+		LLM_SCOPE(ELLMTag::MediaStreaming);
+
+		auto DstTexture = (ID3D12Resource*)SampleDestinationTexture->GetNativeResource();
+
+		D3D12_RESOURCE_BARRIER PreCopyResourceBarriers[] = {
+			{D3D12_RESOURCE_BARRIER_TYPE_TRANSITION, D3D12_RESOURCE_BARRIER_FLAG_NONE, { SampleTexture, D3D12_RESOURCE_BARRIER_ALL_SUBRESOURCES, D3D12_RESOURCE_STATE_COMMON, D3D12_RESOURCE_STATE_COPY_SOURCE }},
+			{D3D12_RESOURCE_BARRIER_TYPE_TRANSITION, D3D12_RESOURCE_BARRIER_FLAG_NONE, { DstTexture, D3D12_RESOURCE_BARRIER_ALL_SUBRESOURCES, D3D12_RESOURCE_STATE_NON_PIXEL_SHADER_RESOURCE | D3D12_RESOURCE_STATE_PIXEL_SHADER_RESOURCE, D3D12_RESOURCE_STATE_COPY_DEST }},
+		};
+
+		D3D12_RESOURCE_BARRIER PostCopyResourceBarriers[] = {
+			{ D3D12_RESOURCE_BARRIER_TYPE_TRANSITION, D3D12_RESOURCE_BARRIER_FLAG_NONE, { SampleTexture, D3D12_RESOURCE_BARRIER_ALL_SUBRESOURCES, D3D12_RESOURCE_STATE_COPY_SOURCE, D3D12_RESOURCE_STATE_COMMON }},
+			{ D3D12_RESOURCE_BARRIER_TYPE_TRANSITION, D3D12_RESOURCE_BARRIER_FLAG_NONE, { DstTexture, D3D12_RESOURCE_BARRIER_ALL_SUBRESOURCES, D3D12_RESOURCE_STATE_COPY_DEST, D3D12_RESOURCE_STATE_NON_PIXEL_SHADER_RESOURCE | D3D12_RESOURCE_STATE_PIXEL_SHADER_RESOURCE }},
+		};
+
+		D3DCmdAllocator->Reset();
+		D3DCmdList->Reset(D3DCmdAllocator, nullptr);
+
+		D3DCmdList->ResourceBarrier(2, PreCopyResourceBarriers);
+		D3DCmdList->CopyResource(DstTexture, SampleTexture);
+		D3DCmdList->ResourceBarrier(2, PostCopyResourceBarriers);
+
+		D3DCmdList->Close();
+	}
+
+	void Execute(FRHICommandListBase& CmdList)
+	{
+		TRefCountPtr<ID3D12GraphicsCommandList> D3DCmdList2 = D3DCmdList;
+		TRefCountPtr<ID3D12CommandAllocator> D3DCmdAllocator2 = D3DCmdAllocator;
+		TRefCountPtr<ID3D12Resource> SampleTexture2 = SampleTexture;
+		TRefCountPtr<ID3D12Fence> D3DFence2 = D3DFence;
+		FTexture2DRHIRef SampleDestinationTexture2 = SampleDestinationTexture;
+		ENQUEUE_RENDER_COMMAND(MediaCopyInputTexture)(
+			[D3DCmdList2, D3DCmdAllocator2, SampleTexture2, SampleDestinationTexture2, D3DFence2, FenceValue2 = FenceValue](FRHICommandListImmediate& RHICmdList)
+		{
+			LLM_SCOPE(ELLMTag::MediaStreaming);
+			auto D3DCmdQueue = (ID3D12CommandQueue*)FRHICommandListExecutor::GetImmediateCommandList().GetNativeGraphicsQueue();
+			if (D3DFence2.IsValid())
+			{
+				D3DCmdQueue->Wait(D3DFence2, FenceValue2);
+			}
+			ID3D12CommandList* CmdLists[1] = {D3DCmdList2.GetReference()};
+			D3DCmdQueue->ExecuteCommandLists(1, CmdLists);
+		});
+	}
+};
+
+
 #ifdef ELECTRA_HAVE_DX11
-struct FRHICommandCopyResource final : public FRHICommand<FRHICommandCopyResource>
+struct FRHICommandCopyResourceDX11 final : public FRHICommand<FRHICommandCopyResourceDX11>
 {
 	TRefCountPtr<ID3D11Texture2D> SampleTexture;
 	FTexture2DRHIRef SampleDestinationTexture;
 	bool bCrossDevice;
 
-	FRHICommandCopyResource(ID3D11Texture2D* InSampleTexture, FRHITexture2D* InSampleDestinationTexture, bool bInCrossDevice)
+	FRHICommandCopyResourceDX11(ID3D11Texture2D* InSampleTexture, FRHITexture2D* InSampleDestinationTexture, bool bInCrossDevice)
 		: SampleTexture(InSampleTexture)
 		, SampleDestinationTexture(InSampleDestinationTexture)
 		, bCrossDevice(bInCrossDevice)
@@ -147,10 +223,8 @@ struct FRHICommandCopyResource final : public FRHICommand<FRHICommandCopyResourc
 /**
  * "Converter" for textures - here: a copy from the decoder owned texture (possibly in another device) into a RHI one (as prep for the real conversion to RGB etc.)
  */
-bool FElectraTextureSample::Convert(FTexture2DRHIRef & InDstTexture, const FConversionHints & Hints)
+bool FElectraTextureSample::Convert(FTexture2DRHIRef& InDstTexture, const FConversionHints& Hints)
 {
-#ifdef ELECTRA_HAVE_DX11
-
 	LLM_SCOPE(ELLMTag::MediaStreaming);
 
 	check(IsInRenderingThread());
@@ -163,6 +237,8 @@ bool FElectraTextureSample::Convert(FTexture2DRHIRef & InDstTexture, const FConv
 	// Get actual sample dimensions
 	FIntPoint Dim = VideoDecoderOutput->GetDim();
 
+	// Note: the converter is not used at all if we have texture data in a SW buffer!
+
 	bool bCrossDeviceCopy;
 	EPixelFormat Format;
 	if (VideoDecoderOutputPC->GetOutputType() != FVideoDecoderOutputPC::EOutputType::HardwareWin8Plus)
@@ -205,22 +281,67 @@ bool FElectraTextureSample::Convert(FTexture2DRHIRef & InDstTexture, const FConv
 
 		Texture = RHICreateTexture(Desc);
 	}
-	
-	// Copy data into RHI texture
-	if (RHICmdList.Bypass())
+
+	uint64 SyncValue = 0;
+	TRefCountPtr<IUnknown> TextureCommon = VideoDecoderOutputPC->GetTexture();
+	TRefCountPtr<IUnknown> SyncCommon = VideoDecoderOutputPC->GetSync(SyncValue);
+
+#ifdef ELECTRA_HAVE_DX11
+	// DX11 texture?
+	TRefCountPtr<ID3D11Texture2D> TextureDX11;
+	HRESULT Res = TextureCommon->QueryInterface(__uuidof(ID3D11Texture2D), (void**)TextureDX11.GetInitReference());
+	if (Res == S_OK)
 	{
-		FRHICommandCopyResource Cmd(VideoDecoderOutputPC->GetTexture(), Texture, bCrossDeviceCopy);
-		Cmd.Execute(RHICmdList);
+		check(RHIGetInterfaceType() == ERHIInterfaceType::D3D11);
+
+		// Copy data into RHI texture
+		if (RHICmdList.Bypass())
+		{
+			FRHICommandCopyResourceDX11 Cmd(TextureDX11, Texture, bCrossDeviceCopy);
+			Cmd.Execute(RHICmdList);
+		}
+		else
+		{
+			new (RHICmdList.AllocCommand<FRHICommandCopyResourceDX11>()) FRHICommandCopyResourceDX11(TextureDX11, Texture, bCrossDeviceCopy);
+		}
 	}
 	else
+#endif
 	{
-		new (RHICmdList.AllocCommand<FRHICommandCopyResource>()) FRHICommandCopyResource(VideoDecoderOutputPC->GetTexture(), Texture, bCrossDeviceCopy);
-	}
+		// No. Then it better be a DX12 resource...
+		TRefCountPtr<ID3D12Resource> TextureDX12;
+		Res = TextureCommon->QueryInterface(__uuidof(ID3D12Resource), (void**)TextureDX12.GetInitReference());
+		check(Res == S_OK);
+		if (Res != S_OK)
+		{
+			return true;
+		}
 	
+		check(RHIGetInterfaceType() == ERHIInterfaceType::D3D12);
+
+		TRefCountPtr<ID3D12Fence> D3DFence;
+		if (SyncCommon)
+		{
+			Res = SyncCommon->QueryInterface(__uuidof(ID3D12Fence), (void**)D3DFence.GetInitReference());
+			check(Res == S_OK);
+			if (Res != S_OK)
+			{
+				return true;
+			}
+		}
+
+		// Copy data into RHI texture
+		if (RHICmdList.Bypass())
+		{
+			FRHICommandCopyResourceDX12 Cmd(TextureDX12, Texture, D3DCmdAllocator, D3DCmdList, D3DFence, SyncValue);
+			Cmd.Execute(RHICmdList);
+		}
+		else
+		{
+			new (RHICmdList.AllocCommand<FRHICommandCopyResourceDX12>()) FRHICommandCopyResourceDX12(TextureDX12, Texture, D3DCmdAllocator, D3DCmdList, D3DFence, SyncValue);
+		}
+	}
 	return true;
-#else
-	return false;
-#endif
 }
 
 
diff --git a/Engine/Plugins/Media/ElectraUtil/Source/ElectraSamples/Public/Windows/ElectraTextureSample.h b/Engine/Plugins/Media/ElectraUtil/Source/ElectraSamples/Public/Windows/ElectraTextureSample.h
--- a/Engine/Plugins/Media/ElectraUtil/Source/ElectraSamples/Public/Windows/ElectraTextureSample.h
+++ b/Engine/Plugins/Media/ElectraUtil/Source/ElectraSamples/Public/Windows/ElectraTextureSample.h
@@ -12,6 +12,10 @@
 #error "Should only be used on Windows"
 #endif
 
+THIRD_PARTY_INCLUDES_START
+#include "d3d12.h"
+THIRD_PARTY_INCLUDES_END
+
 class ELECTRASAMPLES_API FElectraTextureSample final
 	: public IElectraTextureSampleBase
 	, public IMediaTextureSampleConverter
@@ -64,6 +68,10 @@ private:
 
 	/** Output data from video decoder. Baseclass holds reference */
 	FVideoDecoderOutputPC* VideoDecoderOutputPC;
+
+	/** DX12 related resources */
+	TRefCountPtr<ID3D12GraphicsCommandList> D3DCmdList;
+	TRefCountPtr<ID3D12CommandAllocator> D3DCmdAllocator;
 };
 
 using FElectraTextureSamplePtr = TSharedPtr<FElectraTextureSample, ESPMode::ThreadSafe>;
