Fix Clang warnings with ElectraTextureSamplePC and PixelStreaming

diff --git a/Engine/Plugins/Media/ElectraUtil/Source/ElectraSamples/Private/ElectraTextureSamplePC.cpp b/Engine/Plugins/Media/ElectraUtil/Source/ElectraSamples/Private/ElectraTextureSamplePC.cpp
--- a/Engine/Plugins/Media/ElectraUtil/Source/ElectraSamples/Private/ElectraTextureSamplePC.cpp
+++ b/Engine/Plugins/Media/ElectraUtil/Source/ElectraSamples/Private/ElectraTextureSamplePC.cpp
@@ -88,13 +88,13 @@ struct FRHICommandCopyResourceDX12 final : public FRHICommand<FRHICommandCopyRes
 		auto DstTexture = (ID3D12Resource*)SampleDestinationTexture->GetNativeResource();
 
 		D3D12_RESOURCE_BARRIER PreCopyResourceBarriers[] = {
-			{D3D12_RESOURCE_BARRIER_TYPE_TRANSITION, D3D12_RESOURCE_BARRIER_FLAG_NONE, { SampleTexture, D3D12_RESOURCE_BARRIER_ALL_SUBRESOURCES, D3D12_RESOURCE_STATE_COMMON, D3D12_RESOURCE_STATE_COPY_SOURCE }},
-			{D3D12_RESOURCE_BARRIER_TYPE_TRANSITION, D3D12_RESOURCE_BARRIER_FLAG_NONE, { DstTexture, D3D12_RESOURCE_BARRIER_ALL_SUBRESOURCES, D3D12_RESOURCE_STATE_NON_PIXEL_SHADER_RESOURCE | D3D12_RESOURCE_STATE_PIXEL_SHADER_RESOURCE, D3D12_RESOURCE_STATE_COPY_DEST }},
+			{D3D12_RESOURCE_BARRIER_TYPE_TRANSITION, D3D12_RESOURCE_BARRIER_FLAG_NONE, {{ SampleTexture, D3D12_RESOURCE_BARRIER_ALL_SUBRESOURCES, D3D12_RESOURCE_STATE_COMMON, D3D12_RESOURCE_STATE_COPY_SOURCE }}},
+			{D3D12_RESOURCE_BARRIER_TYPE_TRANSITION, D3D12_RESOURCE_BARRIER_FLAG_NONE, {{ DstTexture, D3D12_RESOURCE_BARRIER_ALL_SUBRESOURCES, D3D12_RESOURCE_STATE_NON_PIXEL_SHADER_RESOURCE | D3D12_RESOURCE_STATE_PIXEL_SHADER_RESOURCE, D3D12_RESOURCE_STATE_COPY_DEST }}},
 		};
 
 		D3D12_RESOURCE_BARRIER PostCopyResourceBarriers[] = {
-			{ D3D12_RESOURCE_BARRIER_TYPE_TRANSITION, D3D12_RESOURCE_BARRIER_FLAG_NONE, { SampleTexture, D3D12_RESOURCE_BARRIER_ALL_SUBRESOURCES, D3D12_RESOURCE_STATE_COPY_SOURCE, D3D12_RESOURCE_STATE_COMMON }},
-			{ D3D12_RESOURCE_BARRIER_TYPE_TRANSITION, D3D12_RESOURCE_BARRIER_FLAG_NONE, { DstTexture, D3D12_RESOURCE_BARRIER_ALL_SUBRESOURCES, D3D12_RESOURCE_STATE_COPY_DEST, D3D12_RESOURCE_STATE_NON_PIXEL_SHADER_RESOURCE | D3D12_RESOURCE_STATE_PIXEL_SHADER_RESOURCE }},
+			{ D3D12_RESOURCE_BARRIER_TYPE_TRANSITION, D3D12_RESOURCE_BARRIER_FLAG_NONE, {{ SampleTexture, D3D12_RESOURCE_BARRIER_ALL_SUBRESOURCES, D3D12_RESOURCE_STATE_COPY_SOURCE, D3D12_RESOURCE_STATE_COMMON }}},
+			{ D3D12_RESOURCE_BARRIER_TYPE_TRANSITION, D3D12_RESOURCE_BARRIER_FLAG_NONE, {{ DstTexture, D3D12_RESOURCE_BARRIER_ALL_SUBRESOURCES, D3D12_RESOURCE_STATE_COPY_DEST, D3D12_RESOURCE_STATE_NON_PIXEL_SHADER_RESOURCE | D3D12_RESOURCE_STATE_PIXEL_SHADER_RESOURCE }}},
 		};
 
 		D3DCmdAllocator->Reset();
