Fixing up crash with Electra player when running DX11 on PC
#rnx

diff --git a/Engine/Plugins/Media/ElectraUtil/Source/ElectraSamples/Private/ElectraTextureSamplePC.cpp b/Engine/Plugins/Media/ElectraUtil/Source/ElectraSamples/Private/ElectraTextureSamplePC.cpp
--- a/Engine/Plugins/Media/ElectraUtil/Source/ElectraSamples/Private/ElectraTextureSamplePC.cpp
+++ b/Engine/Plugins/Media/ElectraUtil/Source/ElectraSamples/Private/ElectraTextureSamplePC.cpp
@@ -44,18 +44,21 @@ void FElectraTextureSample::Initialize(FVideoDecoderOutput *InVideoDecoderOutput
 	VideoDecoderOutputPC = static_cast<FVideoDecoderOutputPC*>(InVideoDecoderOutput);
 	SampleFormat = (VideoDecoderOutput->GetFormat() == PF_NV12) ? EMediaTextureSampleFormat::CharNV12 : EMediaTextureSampleFormat::P010;
 
-	ID3D12Device* ApplicationDxDevice = static_cast<ID3D12Device*>(GDynamicRHI->RHIGetNativeDevice());
-	HRESULT Res;
-	if (D3DCmdAllocator.IsValid() == false)
+	if (RHIGetInterfaceType() == ERHIInterfaceType::D3D12)
 	{
-		Res = ApplicationDxDevice->CreateCommandAllocator(D3D12_COMMAND_LIST_TYPE_DIRECT, IID_PPV_ARGS(D3DCmdAllocator.GetInitReference()));
-		check(!FAILED(Res));
-	}
-	if (D3DCmdList.IsValid() == false)
-	{
-		Res = ApplicationDxDevice->CreateCommandList(0, D3D12_COMMAND_LIST_TYPE_DIRECT, D3DCmdAllocator.GetReference(), nullptr, __uuidof(ID3D12CommandList), reinterpret_cast<void**>(D3DCmdList.GetInitReference()));
-		check(!FAILED(Res));
-		D3DCmdList->Close();
+		ID3D12Device* ApplicationDxDevice = static_cast<ID3D12Device*>(GDynamicRHI->RHIGetNativeDevice());
+		HRESULT Res;
+		if (!D3DCmdAllocator.IsValid())
+		{
+			Res = ApplicationDxDevice->CreateCommandAllocator(D3D12_COMMAND_LIST_TYPE_DIRECT, IID_PPV_ARGS(D3DCmdAllocator.GetInitReference()));
+			check(!FAILED(Res));
+		}
+		if (!D3DCmdList.IsValid())
+		{
+			Res = ApplicationDxDevice->CreateCommandList(0, D3D12_COMMAND_LIST_TYPE_DIRECT, D3DCmdAllocator.GetReference(), nullptr, __uuidof(ID3D12CommandList), reinterpret_cast<void**>(D3DCmdList.GetInitReference()));
+			check(!FAILED(Res));
+			D3DCmdList->Close();
+		}
 	}
 }
 
