Adjusting MFW texture sample HDR processing for some console platforms to use the new nomralization factor provided by the tetxure sample itself
#rnx

diff --git a/Engine/Plugins/Media/ElectraUtil/Source/ElectraSamples/Private/ElectraTextureSampleAndroid.cpp b/Engine/Plugins/Media/ElectraUtil/Source/ElectraSamples/Private/ElectraTextureSampleAndroid.cpp
--- a/Engine/Plugins/Media/ElectraUtil/Source/ElectraSamples/Private/ElectraTextureSampleAndroid.cpp
+++ b/Engine/Plugins/Media/ElectraUtil/Source/ElectraSamples/Private/ElectraTextureSampleAndroid.cpp
@@ -517,7 +517,7 @@ bool FElectraTextureSample::Convert(FTexture2DRHIRef& InDstTexture, const FConve
 	if (GetEncodingType() == UE::Color::EEncoding::ST2084)
 	{
 		// Normalize output (e.g. 80 or 100 nits == 1.0)
-		ColorSpaceMtx = ColorSpaceMtx.ApplyScale(kMediaSample_HDR_NitsNormalizationFactor);
+		ColorSpaceMtx = ColorSpaceMtx.ApplyScale(GetHDRNitsNormalizationFactor());
 	}
 
 	// Get shaders.
