Fixing up compile issues for Apple platforms
#rnx

diff --git a/Engine/Plugins/Media/ElectraUtil/Source/ElectraSamples/Private/ElectraTextureSampleAVF.cpp b/Engine/Plugins/Media/ElectraUtil/Source/ElectraSamples/Private/ElectraTextureSampleAVF.cpp
--- a/Engine/Plugins/Media/ElectraUtil/Source/ElectraSamples/Private/ElectraTextureSampleAVF.cpp
+++ b/Engine/Plugins/Media/ElectraUtil/Source/ElectraSamples/Private/ElectraTextureSampleAVF.cpp
@@ -183,9 +183,6 @@ void FElectraMediaTexConvApple::ConvertTexture(FTexture2DRHIRef & InDstTexture,
 
 			bool bIs8Bit = (Format == EMediaTextureSampleFormat::CharNV12);
 
-			const FMatrix* ColorTransform = bFullRange ? &MediaShaders::YuvToRgbRec709Unscaled : &MediaShaders::YuvToRgbRec709Scaled;
-			FVector Off = bFullRange ? MediaShaders::YUVOffsetNoScale8bits : MediaShaders::YUVOffset8bits;
-
 			// Expecting BiPlanar kCVPixelFormatType_420YpCbCr8BiPlanar Full/Video
 			check(CVPixelBufferGetPlaneCount(InImageBufferRef) == 2);
 
@@ -238,20 +235,27 @@ void FElectraMediaTexConvApple::ConvertTexture(FTexture2DRHIRef & InDstTexture,
 					GraphicsPSOInit.BoundShaderState.VertexDeclarationRHI = GMediaVertexDeclaration.VertexDeclarationRHI;
 					GraphicsPSOInit.PrimitiveType = PT_TriangleStrip;
 
+					// Setup conversion from Rec2020 to current working color space
+					const UE::Color::FColorSpace& Working = UE::Color::FColorSpace::GetWorking();
+					FMatrix44f ColorSpaceMtx = UE::Color::Transpose<float>(Working.GetXYZToRgb()) * GamutToXYZMtx;
+					ColorSpaceMtx = ColorSpaceMtx.ApplyScale(NormalizationFactor);
+
 					if (Format == EMediaTextureSampleFormat::CharNV12)
 					{
 						//
 						// NV12
-						// (this disrespects any "working color space")
 						//
 
 						TShaderMapRef<FMediaShadersVS> VertexShader(GlobalShaderMap);
-						TShaderMapRef<FYCbCrConvertPS> PixelShader(GlobalShaderMap);
+						TShaderMapRef<FNV12ConvertPS> PixelShader(GlobalShaderMap);
 						GraphicsPSOInit.BoundShaderState.VertexShaderRHI = VertexShader.GetVertexShader();
 						GraphicsPSOInit.BoundShaderState.PixelShaderRHI = PixelShader.GetPixelShader();
 						SetGraphicsPipelineState(RHICmdList, GraphicsPSOInit, 0);
 
-						SetShaderParametersLegacyPS(RHICmdList, PixelShader, YTex, UVTex, *ColorTransform, Off, true);
+						FShaderResourceViewRHIRef Y_SRV = RHICreateShaderResourceView(YTex, 0, 1, PF_G8);
+						FShaderResourceViewRHIRef UV_SRV = RHICreateShaderResourceView(UVTex, 0, 1, PF_R8G8);
+
+						SetShaderParametersLegacyPS(RHICmdList, PixelShader, FIntPoint(YWidth, YHeight), Y_SRV, UV_SRV, FIntPoint(YWidth, YHeight), YUVMtx, EncodingType == UE::Color::EEncoding::sRGB, ColorSpaceMtx, false);
 					}
 					else
 					{
@@ -259,11 +263,6 @@ void FElectraMediaTexConvApple::ConvertTexture(FTexture2DRHIRef & InDstTexture,
 						// P010
 						//
 
-						// Setup conversion from Rec2020 to current working color space
-						const UE::Color::FColorSpace& Working = UE::Color::FColorSpace::GetWorking();
-						FMatrix44f ColorSpaceMtx = UE::Color::Transpose<float>(Working.GetXYZToRgb()) * GamutToXYZMtx;
-						ColorSpaceMtx = ColorSpaceMtx.ApplyScale(NormalizationFactor);
-
 						TShaderMapRef<FP010ConvertPS> PixelShader(GlobalShaderMap);
 						TShaderMapRef<FMediaShadersVS> VertexShader(GlobalShaderMap);
 						GraphicsPSOInit.BoundShaderState.VertexShaderRHI = VertexShader.GetVertexShader();
@@ -274,7 +273,7 @@ void FElectraMediaTexConvApple::ConvertTexture(FTexture2DRHIRef & InDstTexture,
 						FShaderResourceViewRHIRef UV_SRV = RHICreateShaderResourceView(UVTex, 0, 1, PF_G16R16);
 						
 						// Update shader uniform parameters.
-						SetShaderParametersLegacyPS(RHICmdList, PixelShader, FIntPoint(YWidth, YHeight), Y_SRV, UV_SRV, FIntPoint(YWidth, YHeight), YUVMtx, ColorSpaceMtx, EncodingType == UE::Color::EEncoding::ST2084);
+						SetShaderParametersLegacyPS(RHICmdList, PixelShader, FIntPoint(YWidth, YHeight), Y_SRV, UV_SRV, FIntPoint(YWidth, YHeight), YUVMtx, ColorSpaceMtx, EncodingType == UE::Color::EEncoding::sRGB, EncodingType == UE::Color::EEncoding::ST2084);
 					}
 
 
