Adding missing accessors to display primaries for texture samples
#rnx

diff --git a/Engine/Plugins/Media/ElectraUtil/Source/ElectraSamples/Private/ElectraTextureSampleBase.cpp b/Engine/Plugins/Media/ElectraUtil/Source/ElectraSamples/Private/ElectraTextureSampleBase.cpp
--- a/Engine/Plugins/Media/ElectraUtil/Source/ElectraSamples/Private/ElectraTextureSampleBase.cpp
+++ b/Engine/Plugins/Media/ElectraUtil/Source/ElectraSamples/Private/ElectraTextureSampleBase.cpp
@@ -219,6 +219,45 @@ FVector2f IElectraTextureSampleBase::GetWhitePoint() const
 	return FVector2f(UE::Color::GetWhitePoint(UE::Color::EWhitePoint::CIE1931_D65));
 }
 
+FVector2f IElectraTextureSampleBase::GetDisplayPrimaryRed() const
+{
+	if (auto PinnedHDRInfo = HDRInfo.Pin())
+	{
+		if (auto ColorVolume = PinnedHDRInfo->GetMasteringDisplayColourVolume())
+		{
+			return FVector2f(ColorVolume->display_primaries_x[0], ColorVolume->display_primaries_y[0]);
+		}
+	}
+	// If no HDR info is present, we assume sRGB
+	return FVector2f(0.64, 0.33);
+}
+
+FVector2f IElectraTextureSampleBase::GetDisplayPrimaryGreen() const
+{
+	if (auto PinnedHDRInfo = HDRInfo.Pin())
+	{
+		if (auto ColorVolume = PinnedHDRInfo->GetMasteringDisplayColourVolume())
+		{
+			return FVector2f(ColorVolume->display_primaries_x[1], ColorVolume->display_primaries_y[1]);
+		}
+	}
+	// If no HDR info is present, we assume sRGB
+	return FVector2f(0.30, 0.60);
+}
+
+FVector2f IElectraTextureSampleBase::GetDisplayPrimaryBlue() const
+{
+	if (auto PinnedHDRInfo = HDRInfo.Pin())
+	{
+		if (auto ColorVolume = PinnedHDRInfo->GetMasteringDisplayColourVolume())
+		{
+			return FVector2f(ColorVolume->display_primaries_x[2], ColorVolume->display_primaries_y[2]);
+		}
+	}
+	// If no HDR info is present, we assume sRGB
+	return FVector2f(0.15, 0.06);
+}
+
 UE::Color::EEncoding IElectraTextureSampleBase::GetEncodingType() const
 {
 	if (auto PinnedHDRInfo = HDRInfo.Pin())
diff --git a/Engine/Plugins/Media/ElectraUtil/Source/ElectraSamples/Public/IElectraTextureSample.h b/Engine/Plugins/Media/ElectraUtil/Source/ElectraSamples/Public/IElectraTextureSample.h
--- a/Engine/Plugins/Media/ElectraUtil/Source/ElectraSamples/Public/IElectraTextureSample.h
+++ b/Engine/Plugins/Media/ElectraUtil/Source/ElectraSamples/Public/IElectraTextureSample.h
@@ -50,6 +50,9 @@ public:
 	virtual FMatrix44f GetSampleToRGBMatrix() const override;
 	virtual FMatrix44f GetGamutToXYZMatrix() const override;
 	virtual FVector2f GetWhitePoint() const override;
+	virtual FVector2f GetDisplayPrimaryRed() const override;
+	virtual FVector2f GetDisplayPrimaryGreen() const override;
+	virtual FVector2f GetDisplayPrimaryBlue() const override;
 	virtual UE::Color::EEncoding GetEncodingType() const override;
 
 protected:
diff --git a/Engine/Source/Runtime/Media/Public/IMediaTextureSample.h b/Engine/Source/Runtime/Media/Public/IMediaTextureSample.h
--- a/Engine/Source/Runtime/Media/Public/IMediaTextureSample.h
+++ b/Engine/Source/Runtime/Media/Public/IMediaTextureSample.h
@@ -414,6 +414,21 @@ public:
 		return FVector2f(UE::Color::GetWhitePoint(UE::Color::EWhitePoint::CIE1931_D65));
 	}
 
+	virtual FVector2f GetDisplayPrimaryRed() const
+	{
+		return FVector2f(0.64, 0.33);
+	}
+
+	virtual FVector2f GetDisplayPrimaryGreen() const
+	{
+		return FVector2f(0.30, 0.60);
+	}
+
+	virtual FVector2f GetDisplayPrimaryBlue() const
+	{
+		return FVector2f(0.15, 0.06);
+	}
+
 	/**
 	 * Get EOTF / "Gamma" / encoding type of data
 	 */
