Fixng up Android compilation issues in Electgra texture samples
#rnx

diff --git a/Engine/Plugins/Media/ElectraUtil/Source/ElectraSamples/Private/ElectraTextureSampleAndroid.cpp b/Engine/Plugins/Media/ElectraUtil/Source/ElectraSamples/Private/ElectraTextureSampleAndroid.cpp
--- a/Engine/Plugins/Media/ElectraUtil/Source/ElectraSamples/Private/ElectraTextureSampleAndroid.cpp
+++ b/Engine/Plugins/Media/ElectraUtil/Source/ElectraSamples/Private/ElectraTextureSampleAndroid.cpp
@@ -538,7 +538,7 @@ bool FElectraTextureSample::Convert(FTexture2DRHIRef& InDstTexture, const FConve
 	SetGraphicsPipelineState(RHICmdList, GraphicsPSOInit, 0);
 
 	// Update shader uniform parameters.
-	SetShaderParametersLegacyPS(RHICmdList, PixelShader, InputTexture, OutputDim, GetEncodingType() == UE::Color::EEncoding::sRGB, GetEncodingType() == UE::Color::EEncoding::ST2084, ColorSpaceMtx);
+	SetShaderParametersLegacyPS(RHICmdList, PixelShader, InputTexture, OutputDim, GetEncodingType(), ColorSpaceMtx);
 
 	RHICmdList.SetStreamSource(0, CreateTempMediaVertexBuffer(), 0);
 
