Adding NV12 CPU-side-buffer support for Apple platforms to Electra
#rnx

diff --git a/Engine/Plugins/Media/ElectraUtil/Source/ElectraSamples/Private/ElectraTextureSampleAVF.cpp b/Engine/Plugins/Media/ElectraUtil/Source/ElectraSamples/Private/ElectraTextureSampleAVF.cpp
--- a/Engine/Plugins/Media/ElectraUtil/Source/ElectraSamples/Private/ElectraTextureSampleAVF.cpp
+++ b/Engine/Plugins/Media/ElectraUtil/Source/ElectraSamples/Private/ElectraTextureSampleAVF.cpp
@@ -396,7 +396,8 @@ EMediaTextureSampleFormat FElectraTextureSample::GetFormat() const
             break;
         }
         case    EPixelFormat::PF_BC4:   return EMediaTextureSampleFormat::BC4;
-        default:
+		case    EPixelFormat::PF_NV12:  return EMediaTextureSampleFormat::CharNV21;
+		default:
         {
             check(!"Unsupported pixel format");
         }
