Fixing UV order in NV12 output from VPx decoding for Apple platforms
#rnx

diff --git a/Engine/Plugins/Media/ElectraUtil/Source/ElectraSamples/Private/ElectraTextureSampleAVF.cpp b/Engine/Plugins/Media/ElectraUtil/Source/ElectraSamples/Private/ElectraTextureSampleAVF.cpp
--- a/Engine/Plugins/Media/ElectraUtil/Source/ElectraSamples/Private/ElectraTextureSampleAVF.cpp
+++ b/Engine/Plugins/Media/ElectraUtil/Source/ElectraSamples/Private/ElectraTextureSampleAVF.cpp
@@ -396,7 +396,7 @@ EMediaTextureSampleFormat FElectraTextureSample::GetFormat() const
             break;
         }
         case    EPixelFormat::PF_BC4:   return EMediaTextureSampleFormat::BC4;
-		case    EPixelFormat::PF_NV12:  return EMediaTextureSampleFormat::CharNV21;
+		case    EPixelFormat::PF_NV12:  return EMediaTextureSampleFormat::CharNV12;
 		default:
         {
             check(!"Unsupported pixel format");
