Add SetQueryParams to URLParser in ElectraUtils

diff --git a/Engine/Plugins/Media/ElectraUtil/Source/ElectraBase/Private/URLParser.cpp b/Engine/Plugins/Media/ElectraUtil/Source/ElectraBase/Private/URLParser.cpp
--- a/Engine/Plugins/Media/ElectraUtil/Source/ElectraBase/Private/URLParser.cpp
+++ b/Engine/Plugins/Media/ElectraUtil/Source/ElectraBase/Private/URLParser.cpp
@@ -291,6 +291,19 @@ namespace Electra
 	{
 		GetQueryParams(OutQueryParams, Query, bPerformUrlDecoding, bSameNameReplacesValue);
 	}
+	
+	void FURL_RFC3986::SetQueryParams(const TArray<FQueryParam>& InQueryParams)
+	{
+		Query.Empty();
+		for (int i = 0; i < InQueryParams.Num(); ++i)
+		{
+			Query += InQueryParams[i].Name + '=' + InQueryParams[i].Value;
+			if (i != InQueryParams.Num()-1)
+			{
+				Query.AppendChar(TCHAR('&'));
+			}
+		}
+	}
 
 	// Appends or prepends additional query parameters.
 	void FURL_RFC3986::AddQueryParameters(const FString& InQueryParameters, bool bAppend)
diff --git a/Engine/Plugins/Media/ElectraUtil/Source/ElectraBase/Public/Utilities/URLParser.h b/Engine/Plugins/Media/ElectraUtil/Source/ElectraBase/Public/Utilities/URLParser.h
--- a/Engine/Plugins/Media/ElectraUtil/Source/ElectraBase/Public/Utilities/URLParser.h
+++ b/Engine/Plugins/Media/ElectraUtil/Source/ElectraBase/Public/Utilities/URLParser.h
@@ -57,6 +57,9 @@ namespace Electra
 		// Returns the given query parameter string as a list of name/value pairs.
 		// The parameter string MUST NOT start with a '?'.
 		static void GetQueryParams(TArray<FQueryParam>& OutQueryParams, const FString& InQueryParameters, bool bPerformUrlDecoding, bool bSameNameReplacesValue=true);
+		
+		// Replaces query parameters in the URL.
+        void SetQueryParams(const TArray<FQueryParam>& InQueryParams);
 
 		// Decodes %XX escaped sequences into their original characters. Appends to the output. Hence in and out must not be the same.
 		static bool UrlDecode(FString& OutResult, const FString& InUrlToDecode);
