Electra: Retaining a globalref to the Android surface to work around an otherwise possible race condition in initialization while closing the player

#rnx

diff --git a/Engine/Plugins/Media/ElectraCodecs/Source/ElectraDecoders/Private/Android/h264/VideoDecoderH264_JavaWrapper_Android.cpp b/Engine/Plugins/Media/ElectraCodecs/Source/ElectraDecoders/Private/Android/h264/VideoDecoderH264_JavaWrapper_Android.cpp
--- a/Engine/Plugins/Media/ElectraCodecs/Source/ElectraDecoders/Private/Android/h264/VideoDecoderH264_JavaWrapper_Android.cpp
+++ b/Engine/Plugins/Media/ElectraCodecs/Source/ElectraDecoders/Private/Android/h264/VideoDecoderH264_JavaWrapper_Android.cpp
@@ -465,22 +465,18 @@ int32 FElectraH264VideoDecoderAndroidJava::InitializeDecoder(const FCreateParame
 			JEnv->DeleteLocalRef(CSD1);
 		}
 
-		// Pass along decoder output surface
-		jobject GlobalSurfaceRef = nullptr;
-		if (InCreateParams.VideoCodecSurface)
+		// Pass along decoder output surface. This is a global ref that we take control over now.
+		jobject GlobalSurfaceRef = InCreateParams.VideoCodecSurface;
+		if (GlobalSurfaceRef)
 		{
-			// Make sure we pass a valid global ref.
-			GlobalSurfaceRef = JEnv->NewGlobalRef(InCreateParams.VideoCodecSurface);
-			if (GlobalSurfaceRef)
-			{
-				JEnv->SetObjectField(CreateParams, FCreateParameters_VideoCodecSurface, GlobalSurfaceRef);
-			}
+			JEnv->SetObjectField(CreateParams, FCreateParameters_VideoCodecSurface, GlobalSurfaceRef);
 		}
 		JEnv->SetBooleanField(CreateParams, FCreateParameters_bSurfaceIsView, InCreateParams.bSurfaceIsView);
 
 		// Create and initialize a decoder instance.
 		result = CallMethodNoVerify<int>(ConfigureDecoderFN, CreateParams);
 		JEnv->DeleteLocalRef(CreateParams);
+		// We now release the global ref of the surface we were given. The decoder should have taken ownership of it now.
 		if (GlobalSurfaceRef)
 		{
 			JEnv->DeleteGlobalRef(GlobalSurfaceRef);
diff --git a/Engine/Plugins/Media/ElectraCodecs/Source/ElectraDecoders/Private/Android/h265/VideoDecoderH265_JavaWrapper_Android.cpp b/Engine/Plugins/Media/ElectraCodecs/Source/ElectraDecoders/Private/Android/h265/VideoDecoderH265_JavaWrapper_Android.cpp
--- a/Engine/Plugins/Media/ElectraCodecs/Source/ElectraDecoders/Private/Android/h265/VideoDecoderH265_JavaWrapper_Android.cpp
+++ b/Engine/Plugins/Media/ElectraCodecs/Source/ElectraDecoders/Private/Android/h265/VideoDecoderH265_JavaWrapper_Android.cpp
@@ -459,22 +459,18 @@ int32 FElectraH265VideoDecoderAndroidJava::InitializeDecoder(const FCreateParame
 			JEnv->DeleteLocalRef(CSD0);
 		}
 
-		// Pass along decoder output surface
-		jobject GlobalSurfaceRef = nullptr;
-		if (InCreateParams.VideoCodecSurface)
+		// Pass along decoder output surface. This is a global ref that we take control over now.
+		jobject GlobalSurfaceRef = InCreateParams.VideoCodecSurface;
+		if (GlobalSurfaceRef)
 		{
-			// Make sure we pass a valid global ref.
-			GlobalSurfaceRef = JEnv->NewGlobalRef(InCreateParams.VideoCodecSurface);
-			if (GlobalSurfaceRef)
-			{
-				JEnv->SetObjectField(CreateParams, FCreateParameters_VideoCodecSurface, GlobalSurfaceRef);
-			}
+			JEnv->SetObjectField(CreateParams, FCreateParameters_VideoCodecSurface, GlobalSurfaceRef);
 		}
 		JEnv->SetBooleanField(CreateParams, FCreateParameters_bSurfaceIsView, InCreateParams.bSurfaceIsView);
 
 		// Create and initialize a decoder instance.
 		result = CallMethodNoVerify<int>(ConfigureDecoderFN, CreateParams);
 		JEnv->DeleteLocalRef(CreateParams);
+		// We now release the global ref of the surface we were given. The decoder should have taken ownership of it now.
 		if (GlobalSurfaceRef)
 		{
 			JEnv->DeleteGlobalRef(GlobalSurfaceRef);
diff --git a/Engine/Plugins/Media/ElectraPlayer/Source/ElectraPlayerRuntime/Private/Runtime/Decoder/Android/AndroidElectraDecoderResourceManager.cpp b/Engine/Plugins/Media/ElectraPlayer/Source/ElectraPlayerRuntime/Private/Runtime/Decoder/Android/AndroidElectraDecoderResourceManager.cpp
--- a/Engine/Plugins/Media/ElectraPlayer/Source/ElectraPlayerRuntime/Private/Runtime/Decoder/Android/AndroidElectraDecoderResourceManager.cpp
+++ b/Engine/Plugins/Media/ElectraPlayer/Source/ElectraPlayerRuntime/Private/Runtime/Decoder/Android/AndroidElectraDecoderResourceManager.cpp
@@ -79,6 +79,13 @@ void FElectraDecoderResourceManagerAndroid::FInstanceVars::RequestSurface(TWeakP
 		return;
 	}
 
+	/**
+		Note: The surface returned (if any) must be an _additional_ globalref, not the actual surface.
+		      This is because the actual surface might go away at any time, leaving us with an otherwise stale
+			  reference that will result in a JNI crash when used.
+			  The decoder implementation will pass along this globalref to the decoder and then release it,
+			  so this must not be the actual surface handle!
+	 */
 	void* Surface = rd->VideoDecoderResourceDelegate_GetCodecSurface();
 	SurfaceType = IDecoderPlatformResourceAndroid::ISurfaceRequestCallback::ESurfaceType::Surface;
 	Notify(SurfaceType, Surface);
diff --git a/Engine/Plugins/Media/ElectraUtil/Source/ElectraSamples/Private/ElectraTextureSampleAndroid.cpp b/Engine/Plugins/Media/ElectraUtil/Source/ElectraSamples/Private/ElectraTextureSampleAndroid.cpp
--- a/Engine/Plugins/Media/ElectraUtil/Source/ElectraSamples/Private/ElectraTextureSampleAndroid.cpp
+++ b/Engine/Plugins/Media/ElectraUtil/Source/ElectraSamples/Private/ElectraTextureSampleAndroid.cpp
@@ -86,6 +86,8 @@ private:
 
 	jobject				CodecSurface;
 	FEvent*				SurfaceInitEvent;
+	FCriticalSection	CodecSurfaceLock;
+	jobject				CodecSurfaceToDelete;
 
 
 	static FName GetClassName()
@@ -128,6 +130,7 @@ FElectraTextureSampleSupport::FElectraTextureSampleSupport()
 	, GetVideoFrameUpdateInfoFN(GetClassMethod("GetVideoFrameUpdateInfo", "(IIIZ)Lcom/epicgames/unreal/ElectraTextureSample$FFrameUpdateInfo;"))
 	, CodecSurface(nullptr)
 	, SurfaceInitEvent(nullptr)
+	, CodecSurfaceToDelete(nullptr)
 {
 	JNIEnv* JEnv = FAndroidApplication::GetJavaEnv();
 
@@ -191,6 +194,11 @@ FElectraTextureSampleSupport::~FElectraTextureSampleSupport()
 		SurfaceInitEvent = nullptr;
 	}
 
+	CodecSurfaceLock.Lock();
+	CodecSurfaceToDelete = CodecSurface;
+	CodecSurface = nullptr;
+	CodecSurfaceLock.Unlock();
+
 	if (IsInGameThread())
 	{
 		// enqueue to RT to ensure GL resources are released on the appropriate thread.
@@ -200,9 +208,12 @@ FElectraTextureSampleSupport::~FElectraTextureSampleSupport()
 			{
 				JNIEnv* JEnv = FAndroidApplication::GetJavaEnv();
 				JEnv->CallVoidMethod(Object, ReleaseFN.Method);
-				JEnv->DeleteGlobalRef(CodecSurface);
+				if (CodecSurfaceToDelete)
+				{
+					JEnv->DeleteGlobalRef(CodecSurfaceToDelete);
+				}
 				JEnv->DeleteGlobalRef(FFrameUpdateInfoClass);
-				CodecSurface = 0;
+				CodecSurfaceToDelete = nullptr;
 				FFrameUpdateInfoClass = 0;
 			});
 		});
@@ -212,9 +223,12 @@ FElectraTextureSampleSupport::~FElectraTextureSampleSupport()
 	{
 		JNIEnv* JEnv = FAndroidApplication::GetJavaEnv();
 		JEnv->CallVoidMethod(Object, ReleaseFN.Method);
-		JEnv->DeleteGlobalRef(CodecSurface);
+		if (CodecSurfaceToDelete)
+		{
+			JEnv->DeleteGlobalRef(CodecSurfaceToDelete);
+		}
 		JEnv->DeleteGlobalRef(FFrameUpdateInfoClass);
-		CodecSurface = 0;
+		CodecSurfaceToDelete = nullptr;
 		FFrameUpdateInfoClass = 0;
 	}
 }
@@ -337,7 +351,16 @@ jobject FElectraTextureSampleSupport::GetCodecSurface()
 		}
 	}
 
-	return CodecSurface;
+	// Create a new global ref to return.
+	jobject NewSurfaceHandle = nullptr;
+	CodecSurfaceLock.Lock();
+	if (CodecSurface)
+	{
+		JNIEnv* JEnv = FAndroidApplication::GetJavaEnv();
+		NewSurfaceHandle = JEnv->NewGlobalRef(CodecSurface);
+	}
+	CodecSurfaceLock.Unlock();
+	return NewSurfaceHandle;
 }
 
 //-----------------------------------------------------------------------------
