Fixed local variable shadow error.

diff --git a/Engine/Plugins/Media/ElectraUtil/Source/ElectraSamples/Private/ElectraTextureSampleAVF.cpp b/Engine/Plugins/Media/ElectraUtil/Source/ElectraSamples/Private/ElectraTextureSampleAVF.cpp
--- a/Engine/Plugins/Media/ElectraUtil/Source/ElectraSamples/Private/ElectraTextureSampleAVF.cpp
+++ b/Engine/Plugins/Media/ElectraUtil/Source/ElectraSamples/Private/ElectraTextureSampleAVF.cpp
@@ -155,7 +155,7 @@ public:
 
 void FElectraMediaTexConvApple::ConvertTexture(FTexture2DRHIRef & InDstTexture, CVImageBufferRef InImageBufferRef, bool bFullRange, EMediaTextureSampleFormat Format, const FMatrix44f& YUVMtx, const FMatrix44f& GamutToXYZMtx, UE::Color::EEncoding EncodingType, float NormalizationFactor)
 {
-	FRHICommandListBase& RHICmdList = FRHICommandListImmediate::Get();
+	FRHICommandListImmediate& RHICmdList = FRHICommandListImmediate::Get();
 
 	const int32 FrameHeight = CVPixelBufferGetHeight(InImageBufferRef);
 	const int32 FrameWidth = CVPixelBufferGetWidth(InImageBufferRef);
@@ -217,8 +217,6 @@ void FElectraMediaTexConvApple::ConvertTexture(FTexture2DRHIRef & InDstTexture,
 			TRefCountPtr<FRHITexture> YTex = RHICreateTexture(YDesc);
 			TRefCountPtr<FRHITexture> UVTex = RHICreateTexture(UVDesc);
 
-			// render video frame into sink texture
-			FRHICommandListImmediate& RHICmdList = FRHICommandListExecutor::GetImmediateCommandList();
 			{
 				// configure media shaders
 				auto GlobalShaderMap = GetGlobalShaderMap(GMaxRHIFeatureLevel);
