Use PLATFORM_CURL_INCLUDE to include non-standard/platform-specific versions of the libcurl library.

 - simplified the includes now that the original errors no longer occuring

#rnx

diff --git a/Engine/Plugins/Media/ElectraUtil/Source/ElectraHTTPStream/ElectraHTTPStream.Build.cs b/Engine/Plugins/Media/ElectraUtil/Source/ElectraHTTPStream/ElectraHTTPStream.Build.cs
--- a/Engine/Plugins/Media/ElectraUtil/Source/ElectraHTTPStream/ElectraHTTPStream.Build.cs
+++ b/Engine/Plugins/Media/ElectraUtil/Source/ElectraHTTPStream/ElectraHTTPStream.Build.cs
@@ -125,6 +125,7 @@ public class ElectraHTTPStream : ModuleRules
 			{
 				PublicDefinitions.Add("ELECTRA_HTTPSTREAM_LIBCURL=1");
 				PublicDefinitions.Add("ELECTRA_HTTPSTREAM_XCURL=1");
+				PublicDependencyModuleNames.Add("XCurl");
 			}
 		}
 		else
diff --git a/Engine/Plugins/Media/ElectraUtil/Source/ElectraHTTPStream/Private/Curl/CurlElectra.h b/Engine/Plugins/Media/ElectraUtil/Source/ElectraHTTPStream/Private/Curl/CurlElectra.h
--- a/Engine/Plugins/Media/ElectraUtil/Source/ElectraHTTPStream/Private/Curl/CurlElectra.h
+++ b/Engine/Plugins/Media/ElectraUtil/Source/ElectraHTTPStream/Private/Curl/CurlElectra.h
@@ -10,16 +10,8 @@
 #include "Microsoft/AllowMicrosoftPlatformTypes.h"
 #endif
 
-#if ELECTRA_HTTPSTREAM_XCURL
-//We copied this template to include the windows file from WindowsHWrapper's way if including MinWindows.h, since including xcurl.h directly caused gnarly build errors
-#include "CoreTypes.h"
-#include "HAL/PlatformMemory.h"
-#include "Microsoft/PreWindowsApi.h"
-#ifndef STRICT
-#define STRICT
-#endif
-#include "XCurl.h"
-#include "Microsoft/PostWindowsApi.h"
+#ifdef PLATFORM_CURL_INCLUDE
+	#include PLATFORM_CURL_INCLUDE
 #else
 	#include "curl/curl.h"
 #endif
diff --git a/Engine/Source/Runtime/Online/HTTP/Private/Curl/CurlHttp.h b/Engine/Source/Runtime/Online/HTTP/Private/Curl/CurlHttp.h
--- a/Engine/Source/Runtime/Online/HTTP/Private/Curl/CurlHttp.h
+++ b/Engine/Source/Runtime/Online/HTTP/Private/Curl/CurlHttp.h
@@ -18,16 +18,8 @@ class FCurlHttpResponse;
 #include "Microsoft/AllowMicrosoftPlatformTypes.h"
 #endif
 
-#if WITH_CURL_XCURL
-//We copied this template to include the windows file from WindowsHWrapper's way if including MinWindows.h, since including xcurl.h directly caused gnarly build errors
-#include "CoreTypes.h"
-#include "HAL/PlatformMemory.h"
-#include "Microsoft/PreWindowsApi.h"
-#ifndef STRICT
-#define STRICT
-#endif
-#include "xcurl.h"
-#include "Microsoft/PostWindowsApi.h"
+#ifdef PLATFORM_CURL_INCLUDE
+	#include PLATFORM_CURL_INCLUDE
 #else
 	#include "curl/curl.h"
 #endif
diff --git a/Engine/Source/Runtime/Online/HTTP/Private/Curl/CurlHttpThread.h b/Engine/Source/Runtime/Online/HTTP/Private/Curl/CurlHttpThread.h
--- a/Engine/Source/Runtime/Online/HTTP/Private/Curl/CurlHttpThread.h
+++ b/Engine/Source/Runtime/Online/HTTP/Private/Curl/CurlHttpThread.h
@@ -11,16 +11,8 @@
 #if PLATFORM_MICROSOFT
 #include "Microsoft/AllowMicrosoftPlatformTypes.h"
 #endif
-#if WITH_CURL_XCURL
-//We copied this template to include the windows file from WindowsHWrapper's way if including MinWindows.h, since including xcurl.h directly caused gnarly build errors
-#include "CoreTypes.h"
-#include "HAL/PlatformMemory.h"
-#include "Microsoft/PreWindowsApi.h"
-#ifndef STRICT
-#define STRICT
-#endif
-#include "xcurl.h"
-#include "Microsoft/PostWindowsApi.h"
+#ifdef PLATFORM_CURL_INCLUDE
+#include PLATFORM_CURL_INCLUDE
 #else
 #include "curl/curl.h"
 #endif
diff --git a/Engine/Source/Runtime/Online/HTTP/Private/Curl/CurlMultiPollEventLoopHttpThread.h b/Engine/Source/Runtime/Online/HTTP/Private/Curl/CurlMultiPollEventLoopHttpThread.h
--- a/Engine/Source/Runtime/Online/HTTP/Private/Curl/CurlMultiPollEventLoopHttpThread.h
+++ b/Engine/Source/Runtime/Online/HTTP/Private/Curl/CurlMultiPollEventLoopHttpThread.h
@@ -14,16 +14,8 @@
 #if PLATFORM_MICROSOFT
 #include "Microsoft/AllowMicrosoftPlatformTypes.h"
 #endif
-#if WITH_CURL_XCURL
-//We copied this template to include the windows file from WindowsHWrapper's way if including MinWindows.h, since including xcurl.h directly caused gnarly build errors
-#include "CoreTypes.h"
-#include "HAL/PlatformMemory.h"
-#include "Microsoft/PreWindowsApi.h"
-#ifndef STRICT
-#define STRICT
-#endif
-#include "xcurl.h"
-#include "Microsoft/PostWindowsApi.h"
+#ifdef PLATFORM_CURL_INCLUDE
+#include PLATFORM_CURL_INCLUDE
 #else
 #include "curl/curl.h"
 #endif
diff --git a/Engine/Source/Runtime/Online/HTTP/Private/Curl/CurlMultiWaitEventLoopHttpThread.h b/Engine/Source/Runtime/Online/HTTP/Private/Curl/CurlMultiWaitEventLoopHttpThread.h
--- a/Engine/Source/Runtime/Online/HTTP/Private/Curl/CurlMultiWaitEventLoopHttpThread.h
+++ b/Engine/Source/Runtime/Online/HTTP/Private/Curl/CurlMultiWaitEventLoopHttpThread.h
@@ -14,16 +14,8 @@
 #if PLATFORM_MICROSOFT
 #include "Microsoft/AllowMicrosoftPlatformTypes.h"
 #endif
-#if WITH_CURL_XCURL
-//We copied this template to include the windows file from WindowsHWrapper's way if including MinWindows.h, since including xcurl.h directly caused gnarly build errors
-#include "CoreTypes.h"
-#include "HAL/PlatformMemory.h"
-#include "Microsoft/PreWindowsApi.h"
-#ifndef STRICT
-#define STRICT
-#endif
-#include "xcurl.h"
-#include "Microsoft/PostWindowsApi.h"
+#ifdef PLATFORM_CURL_INCLUDE
+#include PLATFORM_CURL_INCLUDE
 #else
 #include "curl/curl.h"
 #endif
diff --git a/Engine/Source/Runtime/Online/HTTP/Private/Curl/CurlSocketEventLoopHttpThread.h b/Engine/Source/Runtime/Online/HTTP/Private/Curl/CurlSocketEventLoopHttpThread.h
--- a/Engine/Source/Runtime/Online/HTTP/Private/Curl/CurlSocketEventLoopHttpThread.h
+++ b/Engine/Source/Runtime/Online/HTTP/Private/Curl/CurlSocketEventLoopHttpThread.h
@@ -15,16 +15,8 @@
 #if PLATFORM_MICROSOFT
 #include "Microsoft/AllowMicrosoftPlatformTypes.h"
 #endif
-#if WITH_CURL_XCURL
-//We copied this template to include the windows file from WindowsHWrapper's way if including MinWindows.h, since including xcurl.h directly caused gnarly build errors
-#include "CoreTypes.h"
-#include "HAL/PlatformMemory.h"
-#include "Microsoft/PreWindowsApi.h"
-#ifndef STRICT
-#define STRICT
-#endif
-#include "xcurl.h"
-#include "Microsoft/PostWindowsApi.h"
+#ifdef PLATFORM_CURL_INCLUDE
+#include PLATFORM_CURL_INCLUDE
 #else
 #include "curl/curl.h"
 #endif
