Extending support for colorimetry information from video streams for non-standard HDR streams
(e.g. allowing 8-bit HEVC data to use Rec2020 & PQ etc.)

diff --git a/Engine/Plugins/Media/ElectraUtil/Source/ElectraSamples/Private/ElectraTextureSampleBase.cpp b/Engine/Plugins/Media/ElectraUtil/Source/ElectraSamples/Private/ElectraTextureSampleBase.cpp
--- a/Engine/Plugins/Media/ElectraUtil/Source/ElectraSamples/Private/ElectraTextureSampleBase.cpp
+++ b/Engine/Plugins/Media/ElectraUtil/Source/ElectraSamples/Private/ElectraTextureSampleBase.cpp
@@ -1,10 +1,13 @@
 // Copyright Epic Games, Inc. All Rights Reserved.
 
-#include "CoreMinimal.h"
+#include "Misc/Build.h"
 
 #if !UE_SERVER
 
+#include "HAL/Platform.h"
 #include "ElectraTextureSample.h"
+#include "ElectraTextureSampleUtils.h"
+#include "ElectraSamplesModule.h"
 
 // -------------------------------------------------------------------------------------------------------------------------------------------------------
 
@@ -14,12 +17,20 @@ void IElectraTextureSampleBase::Initialize(FVideoDecoderOutput* InVideoDecoderOu
 	HDRInfo = VideoDecoderOutput->GetHDRInformation();
 	Colorimetry = VideoDecoderOutput->GetColorimetry();
 
+	// Get various basic MP4-style colorimetry values (we default to video range Rec709 SDR)
 	bool bFullRange = false;
+	uint8 ColorPrimaries = ElectraColorimetryUtils::DefaultMPEGColorPrimaries;
+	uint8 TransferCharacteristics = ElectraColorimetryUtils::DefaultMPEGMatrixCoefficients;
+	uint8 MatrixCoefficients = ElectraColorimetryUtils::DefaultMPEGTransferCharacteristics;
 	if (auto PinnedColorimetry = Colorimetry.Pin())
 	{
 		bFullRange = (PinnedColorimetry->GetMPEGDefinition()->VideoFullRangeFlag != 0);
+		ColorPrimaries = PinnedColorimetry->GetMPEGDefinition()->ColourPrimaries;
+		TransferCharacteristics = PinnedColorimetry->GetMPEGDefinition()->TransferCharacteristics;
+		MatrixCoefficients = PinnedColorimetry->GetMPEGDefinition()->MatrixCoefficients;
 	}
-	
+
+	// Compute the bits per component in the data we get passed in
 	EPixelFormat PixFmt = VideoDecoderOutput->GetFormat();
 	uint8 NumBits = 8;
 	if (!IsDXTCBlockCompressedTextureFormat(PixFmt))
@@ -42,50 +53,123 @@ void IElectraTextureSampleBase::Initialize(FVideoDecoderOutput* InVideoDecoderOu
 		}
 	}
 
-	// Prepare YUV -> RGB matrix containing all necessary offsets and scales to produce RGB straight from sample data
-	const FMatrix* Mtx = bFullRange ? &MediaShaders::YuvToRgbRec709Unscaled : &MediaShaders::YuvToRgbRec709Scaled;
-	FVector Off;
-	switch (NumBits)
-	{
-		case 8:		Off = bFullRange ? MediaShaders::YUVOffsetNoScale8bits  : MediaShaders::YUVOffset8bits; break;
-		case 10:	Off = bFullRange ? MediaShaders::YUVOffsetNoScale10bits : MediaShaders::YUVOffset10bits; break;
-		case 32:	Off = bFullRange ? MediaShaders::YUVOffsetNoScaleFloat  : MediaShaders::YUVOffsetFloat; break;
-		default:	Off = bFullRange ? MediaShaders::YUVOffsetNoScale16bits : MediaShaders::YUVOffset16bits; break;
-	}
-	// Correctional scale for input data
-	// (data should be placed in the upper 10-bits of the 16-bit texture channels, but some platforms do not do this)
-	float DataScale = GetSampleDataScale(false);
+	FVector Off = FVector::Zero();
+	const FMatrix* Mtx = nullptr;
 
+	// Do we have specific HDR information, so we can assume a standard?
 	if (auto PinnedHDRInfo = HDRInfo.Pin())
 	{
+		//
+		// HDR information present
+		//
+
+		if (auto ColorVolume = PinnedHDRInfo->GetMasteringDisplayColourVolume())
+		{
+			WhitePoint = FVector2d(ColorVolume->white_point_x, ColorVolume->white_point_y);
+			for (uint32 I = 0; I < 3; ++I)
+			{
+				DisplayPrimaries[I] = FVector2d(ColorVolume->display_primaries_x[I], ColorVolume->display_primaries_y[I]);
+			}
+		}
+		else
+		{
+			WhitePoint = FVector2d(UE::Color::GetWhitePoint(UE::Color::EWhitePoint::CIE1931_D65));
+
+			const FVector2d* DPs = ElectraColorimetryUtils::GetColorPrimaries(UE::Color::EColorSpace::Rec2020);
+			DisplayPrimaries[0] = DPs[0];
+			DisplayPrimaries[1] = DPs[1];
+			DisplayPrimaries[2] = DPs[2];
+		}
+
 		switch (PinnedHDRInfo->GetHDRType())
 		{
 		case IVideoDecoderHDRInformation::EType::PQ10:
-		case IVideoDecoderHDRInformation::EType::HLG10:
 		case IVideoDecoderHDRInformation::EType::HDR10:
-			Mtx = &MediaShaders::YuvToRgbRec2020Unscaled;
-			Off = MediaShaders::YUVOffsetNoScale16bits;
-			DataScale = GetSampleDataScale(true);
+			ColorEncoding = UE::Color::EEncoding::ST2084;
+			break;
+		case IVideoDecoderHDRInformation::EType::HLG10:
+			UE_LOG(LogElectraSamples, Warning, TEXT("Detected use of HLG EOTF in video material. Mapping to Rec709/sRGB as HLG is not supported."));
+			ColorEncoding = UE::Color::EEncoding::sRGB;	// using sRGB in place of HLG for now
 			break;
 		default:
+			check(!"Unknown HDR type!");
+			ColorEncoding = UE::Color::EEncoding::sRGB;
 			break;
 		}
+
+		ColorGamut = EDisplayColorGamut::Rec2020_D65;
+		Mtx = bFullRange ? &MediaShaders::YuvToRgbRec2020Unscaled : &MediaShaders::YuvToRgbRec2020Scaled;
 	}
+	else
+	{
+		//
+		// No HDR information present
+		//
 
-	// Compute scale to make correct towards the max value (1.0 if it is 0xff or 0xffff, but >1.0 if it 0xffc0 - which is the case for P010)
-	float NormScale = (VideoDecoderOutput->GetFormat() == PF_P010) ? (65535.0f / 65472.0f) : 1.0f;
+		const FVector2d* DPs = ElectraColorimetryUtils::GetColorPrimaries(ElectraColorimetryUtils::TranslateMPEGColorPrimaries(ColorPrimaries));
+		DisplayPrimaries[0] = DPs[0];
+		DisplayPrimaries[1] = DPs[1];
+		DisplayPrimaries[2] = DPs[2];
+
+		switch (ElectraColorimetryUtils::TranslateMPEGMatrixCoefficients(MatrixCoefficients))
+		{
+			case UE::Color::EColorSpace::None:	// ID (RGB)
+				// no conversion, data is RGB
+				break;
+			case UE::Color::EColorSpace::sRGB:
+				Mtx = bFullRange ? &MediaShaders::YuvToRgbRec709Unscaled : &MediaShaders::YuvToRgbRec709Scaled;
+				ColorGamut = EDisplayColorGamut::sRGB_D65;
+				break;
+			case UE::Color::EColorSpace::Rec2020:
+				Mtx = bFullRange ? &MediaShaders::YuvToRgbRec2020Unscaled : &MediaShaders::YuvToRgbRec2020Scaled;
+				ColorGamut = EDisplayColorGamut::Rec2020_D65;
+				break;
+			default:
+				check(!"*** Unexpected matrix coefficients!");
+				Mtx = bFullRange ? &MediaShaders::YuvToRgbRec709Unscaled : &MediaShaders::YuvToRgbRec709Scaled;
+				ColorGamut = EDisplayColorGamut::sRGB_D65;
+		}
 
+		ColorEncoding = ElectraColorimetryUtils::TranslateMPEGTransferCharacteristics(TransferCharacteristics);
+
+		// No matter what: we will have the D65 white point (for now at least)
+		WhitePoint = FVector2d(UE::Color::GetWhitePoint(UE::Color::EWhitePoint::CIE1931_D65));
+	}
+
+	if (Mtx)
+	{
+		// Select the offsets prior to YUV conversion needed per the incoming data
+		switch (NumBits)
+		{
+		case 8:		Off = bFullRange ? MediaShaders::YUVOffsetNoScale8bits : MediaShaders::YUVOffset8bits; break;
+		case 10:	Off = bFullRange ? MediaShaders::YUVOffsetNoScale10bits : MediaShaders::YUVOffset10bits; break;
+		case 16:	Off = bFullRange ? MediaShaders::YUVOffsetNoScale16bits : MediaShaders::YUVOffset16bits; break;
+		case 32:	Off = bFullRange ? MediaShaders::YUVOffsetNoScaleFloat : MediaShaders::YUVOffsetFloat; break;
+		default:	check(!"Unexpected number of bits per channel!");
+		}
+	}
+
+	// Correctional scale for input data
+	// (data should be placed in the upper 10-bits of the 16-bit texture channels, but some platforms do not do this - they provide a correctional factor here)
+	float DataScale = GetSampleDataScale(NumBits == 10);
+
+	// Compute scale to make correct towards the max value (P010 will max out at 0xffc0 not 0xffff - so if it is present we need to adjust the scale a bit)
+	float NormScale = (VideoDecoderOutput->GetFormat() == PF_P010) ? (65535.0f / 65472.0f) : 1.0f;
+ 
 	// Matrix to transform sample data to standard YUV values
 	FMatrix PreMtx = FMatrix::Identity;
 	PreMtx.M[0][0] = DataScale * NormScale;
 	PreMtx.M[1][1] = DataScale * NormScale;
 	PreMtx.M[2][2] = DataScale * NormScale;
-	PreMtx.M[0][3] = -Off.X * NormScale;
-	PreMtx.M[1][3] = -Off.Y * NormScale;
-	PreMtx.M[2][3] = -Off.Z * NormScale;
+	PreMtx.M[0][3] = -Off.X;
+	PreMtx.M[1][3] = -Off.Y;
+	PreMtx.M[2][3] = -Off.Z;
 
 	// Combine this with the actual YUV-RGB conversion
-	YuvToRgbMtx = FMatrix44f(*Mtx * PreMtx);
+	SampleToRgbMtx = FMatrix44f(Mtx ? (*Mtx * PreMtx) : PreMtx);
+
+	// Also store the plain YUV->RGB matrix (pointer) for later reference
+	YuvToRgbMtx = Mtx;
 }
 
 void IElectraTextureSampleBase::InitializePoolable()
@@ -139,143 +223,56 @@ FTimespan IElectraTextureSampleBase::GetDuration() const
 
 bool IElectraTextureSampleBase::IsOutputSrgb() const
 {
-	if (auto PinnedHDRInfo = HDRInfo.Pin())
-	{
-		// If the HDR type is unknown we also assume sRGB
-		return (PinnedHDRInfo->GetHDRType() == IVideoDecoderHDRInformation::EType::Unknown);
-	}
-	// If no HDR info is present, we assume sRGB
-	return true;
+	return ColorEncoding == UE::Color::EEncoding::sRGB;
 }
 
 const FMatrix& IElectraTextureSampleBase::GetYUVToRGBMatrix() const
 {
-	if (auto PinnedHDRInfo = HDRInfo.Pin())
-	{
-		switch (PinnedHDRInfo->GetHDRType())
-		{
-		case IVideoDecoderHDRInformation::EType::PQ10:
-		case IVideoDecoderHDRInformation::EType::HLG10:
-		case IVideoDecoderHDRInformation::EType::HDR10:		return MediaShaders::YuvToRgbRec2020Unscaled;
-		case IVideoDecoderHDRInformation::EType::Unknown:	break;
-		}
-	}
-
-	// If no HDR info is present, we assume sRGB
-	return GetFullRange() ? MediaShaders::YuvToRgbRec709Unscaled : MediaShaders::YuvToRgbRec709Scaled;
+	return YuvToRgbMtx ? *YuvToRgbMtx : FMatrix::Identity;
 }
 
 bool IElectraTextureSampleBase::GetFullRange() const
 {
 	if (auto PinnedColorimetry = Colorimetry.Pin())
 	{
-		if (auto PinnedHDRInfo = HDRInfo.Pin())
-		{
-			if (PinnedHDRInfo->GetHDRType() != IVideoDecoderHDRInformation::EType::Unknown)
-			{
-				// For HDR we assume full range at all times
-				return true;
-			}
-		}
-		// SDR honors the flag in the MPEG stream...
 		return (PinnedColorimetry->GetMPEGDefinition()->VideoFullRangeFlag != 0);
 	}
-	// SDR with no MPEG info is assumed to be video range
 	return false;
 }
 
 FMatrix44f IElectraTextureSampleBase::GetSampleToRGBMatrix() const
 {
-	return YuvToRgbMtx;
+	return SampleToRgbMtx;
 }
 
 FMatrix44f IElectraTextureSampleBase::GetGamutToXYZMatrix() const
 {
-	if (auto PinnedHDRInfo = HDRInfo.Pin())
-	{
-		switch (PinnedHDRInfo->GetHDRType())
-		{
-		case IVideoDecoderHDRInformation::EType::PQ10:
-		case IVideoDecoderHDRInformation::EType::HLG10:
-		case IVideoDecoderHDRInformation::EType::HDR10:		return GamutToXYZMatrix(EDisplayColorGamut::Rec2020_D65);
-		case IVideoDecoderHDRInformation::EType::Unknown:	return GamutToXYZMatrix(EDisplayColorGamut::sRGB_D65);
-		}
-	}
-
-	// If no HDR info is present, we assume sRGB
-	return GamutToXYZMatrix(EDisplayColorGamut::sRGB_D65);
+	return GamutToXYZMatrix(ColorGamut);
 }
 
 FVector2f IElectraTextureSampleBase::GetWhitePoint() const
 {
-	if (auto PinnedHDRInfo = HDRInfo.Pin())
-	{
-		if (auto ColorVolume = PinnedHDRInfo->GetMasteringDisplayColourVolume())
-		{
-			return FVector2f(ColorVolume->white_point_x, ColorVolume->white_point_y);
-		}
-	}
-	// If no HDR info is present, we assume sRGB
-	return FVector2f(UE::Color::GetWhitePoint(UE::Color::EWhitePoint::CIE1931_D65));
+	return FVector2f(WhitePoint);
 }
 
 FVector2f IElectraTextureSampleBase::GetDisplayPrimaryRed() const
 {
-	if (auto PinnedHDRInfo = HDRInfo.Pin())
-	{
-		if (auto ColorVolume = PinnedHDRInfo->GetMasteringDisplayColourVolume())
-		{
-			return FVector2f(ColorVolume->display_primaries_x[0], ColorVolume->display_primaries_y[0]);
-		}
-	}
-	// If no HDR info is present, we assume sRGB
-	return FVector2f(0.64, 0.33);
+	return FVector2f(DisplayPrimaries[0]);
 }
 
 FVector2f IElectraTextureSampleBase::GetDisplayPrimaryGreen() const
 {
-	if (auto PinnedHDRInfo = HDRInfo.Pin())
-	{
-		if (auto ColorVolume = PinnedHDRInfo->GetMasteringDisplayColourVolume())
-		{
-			return FVector2f(ColorVolume->display_primaries_x[1], ColorVolume->display_primaries_y[1]);
-		}
-	}
-	// If no HDR info is present, we assume sRGB
-	return FVector2f(0.30, 0.60);
+	return FVector2f(DisplayPrimaries[1]);
 }
 
 FVector2f IElectraTextureSampleBase::GetDisplayPrimaryBlue() const
 {
-	if (auto PinnedHDRInfo = HDRInfo.Pin())
-	{
-		if (auto ColorVolume = PinnedHDRInfo->GetMasteringDisplayColourVolume())
-		{
-			return FVector2f(ColorVolume->display_primaries_x[2], ColorVolume->display_primaries_y[2]);
-		}
-	}
-	// If no HDR info is present, we assume sRGB
-	return FVector2f(0.15, 0.06);
+	return FVector2f(DisplayPrimaries[2]);
 }
 
 UE::Color::EEncoding IElectraTextureSampleBase::GetEncodingType() const
 {
-	if (auto PinnedHDRInfo = HDRInfo.Pin())
-	{
-		switch (PinnedHDRInfo->GetHDRType())
-		{
-		case IVideoDecoderHDRInformation::EType::PQ10:
-		case IVideoDecoderHDRInformation::EType::HDR10:		return UE::Color::EEncoding::ST2084;
-		case IVideoDecoderHDRInformation::EType::Unknown:	return UE::Color::EEncoding::sRGB;
-		case IVideoDecoderHDRInformation::EType::HLG10:
-			{
-			check(!"*** Implement support for HLG10 in UE::Color!");
-			return UE::Color::EEncoding::sRGB;
-			}
-		}
-	}
-	// If no HDR info is present, we assume sRGB
-	return UE::Color::EEncoding::sRGB;
+	return ColorEncoding;
 }
 
 #endif
diff --git a/Engine/Plugins/Media/ElectraUtil/Source/ElectraSamples/Public/ElectraTextureSampleUtils.h b/Engine/Plugins/Media/ElectraUtil/Source/ElectraSamples/Public/ElectraTextureSampleUtils.h
new file mode 100644
--- /dev/null
+++ b/Engine/Plugins/Media/ElectraUtil/Source/ElectraSamples/Public/ElectraTextureSampleUtils.h
@@ -0,0 +1,84 @@
+// Copyright Epic Games, Inc. All Rights Reserved.
+
+#pragma once
+
+#include "HAL/Platform.h"
+#include "ColorManagementDefines.h"
+
+namespace ElectraColorimetryUtils
+{
+	static constexpr uint8 DefaultMPEGColorPrimaries = 2;
+	static constexpr uint8 DefaultMPEGMatrixCoefficients = 2;
+	static constexpr uint8 DefaultMPEGTransferCharacteristics = 2;
+
+	UE::Color::EColorSpace TranslateMPEGColorPrimaries(uint8 InPrimaries)
+	{
+		switch (InPrimaries)
+		{
+			case 1:				// Rec709
+			case 2:				// unknown
+				return UE::Color::EColorSpace::sRGB;
+			case 9:				// Rec2020
+				return UE::Color::EColorSpace::Rec2020;
+			default:
+				check(!"Unexpected MPEG color primaries value!");
+		}
+		return UE::Color::EColorSpace::sRGB;
+	}
+
+	UE::Color::EColorSpace TranslateMPEGMatrixCoefficients(uint8 InMatrixCoefficients)
+	{
+		switch (InMatrixCoefficients)
+		{
+			case 0:				// ID (RGB)
+				return UE::Color::EColorSpace::None;
+			case 1:				// Rec709
+			case 2:				// unknown
+				return UE::Color::EColorSpace::sRGB;
+			case 9:				// Rec2020
+				return UE::Color::EColorSpace::Rec2020;
+			default:
+				check(!"Unexpected MPEG color primaries value!");
+		}
+		return UE::Color::EColorSpace::sRGB;
+	}
+
+	UE::Color::EEncoding TranslateMPEGTransferCharacteristics(uint8 InTransferCharacteristics)
+	{
+		switch (InTransferCharacteristics)
+		{
+			case 1:
+			case 6:
+			case 14:
+			case 15:
+			case 2:				// unknown
+				return UE::Color::EEncoding::sRGB;
+			case 8:
+				return UE::Color::EEncoding::Linear;
+			case 16:
+				return UE::Color::EEncoding::ST2084;
+			case 18:
+				return UE::Color::EEncoding::sRGB;	// using sRGB in place of HLG for now
+			default:
+				check(!"*** Unexpected MPEG color transfer characteristics!");
+		}
+		return UE::Color::EEncoding::sRGB;
+	}
+
+	const FVector2d* GetColorPrimaries(UE::Color::EColorSpace InColorSpace)
+	{
+		static const FVector2d DP_sRGB[3] = { {0.64, 0.33} , {0.30, 0.60}, {0.15, 0.06} };
+		static const FVector2d DP_Rec2020[3] = { {0.708, 0.292} , {0.170, 0.797}, {0.131, 0.046} };
+
+		switch (InColorSpace)
+		{
+		case UE::Color::EColorSpace::sRGB:
+			return DP_sRGB;
+		case UE::Color::EColorSpace::Rec2020:
+			return DP_Rec2020;
+		default:
+			check(!"*** Unexpected color primaries!");
+		}
+		return DP_sRGB;
+	}
+};
diff --git a/Engine/Plugins/Media/ElectraUtil/Source/ElectraSamples/Public/IElectraTextureSample.h b/Engine/Plugins/Media/ElectraUtil/Source/ElectraSamples/Public/IElectraTextureSample.h
--- a/Engine/Plugins/Media/ElectraUtil/Source/ElectraSamples/Public/IElectraTextureSample.h
+++ b/Engine/Plugins/Media/ElectraUtil/Source/ElectraSamples/Public/IElectraTextureSample.h
@@ -66,5 +66,14 @@ protected:
 	TWeakPtr<const IVideoDecoderColorimetry, ESPMode::ThreadSafe> Colorimetry;
 
 	/** YUV matrix, adjusted to compensate for decoder output specific scale */
-	FMatrix44f YuvToRgbMtx;
+	FMatrix44f SampleToRgbMtx;
+
+	/** YUV to RGB matrix without any adjustments for decoder output specifics */
+	const FMatrix* YuvToRgbMtx;
+
+	/** Precomputed colorimetric data */
+	EDisplayColorGamut ColorGamut;
+	UE::Color::EEncoding ColorEncoding;
+	FVector2d WhitePoint;
+	FVector2d DisplayPrimaries[3];
 };
diff --git a/Engine/Source/Runtime/MediaAssets/Private/Misc/MediaTextureResource.cpp b/Engine/Source/Runtime/MediaAssets/Private/Misc/MediaTextureResource.cpp
--- a/Engine/Source/Runtime/MediaAssets/Private/Misc/MediaTextureResource.cpp
+++ b/Engine/Source/Runtime/MediaAssets/Private/Misc/MediaTextureResource.cpp
@@ -277,10 +277,16 @@ namespace MediaTextureResourceHelpers
 		case EMediaTextureSampleFormat::ARGB16_BIG:
 			return PF_A16B16G16R16;
 
-		// Everything else maps to 8-bit RGB...
+		// Everything else maps to 8-bit RGB if it is linear or sRGB (we use sRGB encoding in HW). Otherwise: half floats...
 		default:
+		{
+			if (Sample->GetEncodingType() != UE::Color::EEncoding::Linear && Sample->GetEncodingType() != UE::Color::EEncoding::sRGB)
+			{
+				return PF_FloatRGB;
+			}
 			return PF_B8G8R8A8;
 		}
+		}
 	}
 
 
