Cleanup in Mediaframework chromaticity & HDR related texture sample APIs & Electra internal data representation (now using FColorSpace)
#rnx

diff --git a/Engine/Plugins/Media/ElectraUtil/Source/ElectraSamples/Private/ElectraTextureSampleAVF.cpp b/Engine/Plugins/Media/ElectraUtil/Source/ElectraSamples/Private/ElectraTextureSampleAVF.cpp
--- a/Engine/Plugins/Media/ElectraUtil/Source/ElectraSamples/Private/ElectraTextureSampleAVF.cpp
+++ b/Engine/Plugins/Media/ElectraUtil/Source/ElectraSamples/Private/ElectraTextureSampleAVF.cpp
@@ -153,7 +153,7 @@ public:
 } // namespace anonymous
 
 
-void FElectraMediaTexConvApple::ConvertTexture(FTexture2DRHIRef & InDstTexture, CVImageBufferRef InImageBufferRef, bool bFullRange, EMediaTextureSampleFormat Format, const FMatrix44f& YUVMtx, const FMatrix44f& GamutToXYZMtx, UE::Color::EEncoding EncodingType, float NormalizationFactor)
+void FElectraMediaTexConvApple::ConvertTexture(FTexture2DRHIRef & InDstTexture, CVImageBufferRef InImageBufferRef, bool bFullRange, EMediaTextureSampleFormat Format, const FMatrix44f& YUVMtx, const FMatrix44d& GamutToXYZMtx, UE::Color::EEncoding EncodingType, float NormalizationFactor)
 {
 	FRHICommandListImmediate& RHICmdList = FRHICommandListImmediate::Get();
 
@@ -235,7 +235,7 @@ void FElectraMediaTexConvApple::ConvertTexture(FTexture2DRHIRef & InDstTexture,
 
 					// Setup conversion from Rec2020 to current working color space
 					const UE::Color::FColorSpace& Working = UE::Color::FColorSpace::GetWorking();
-					FMatrix44f ColorSpaceMtx = UE::Color::Transpose<float>(Working.GetXYZToRgb()) * GamutToXYZMtx;
+					FMatrix44f ColorSpaceMtx = FMatrix44f(Working.GetXYZToRgb().GetTransposed() * GamutToXYZMtx);
 					ColorSpaceMtx = ColorSpaceMtx.ApplyScale(NormalizationFactor);
 
 					if (Format == EMediaTextureSampleFormat::CharNV12)
diff --git a/Engine/Plugins/Media/ElectraUtil/Source/ElectraSamples/Private/ElectraTextureSampleAndroid.cpp b/Engine/Plugins/Media/ElectraUtil/Source/ElectraSamples/Private/ElectraTextureSampleAndroid.cpp
--- a/Engine/Plugins/Media/ElectraUtil/Source/ElectraSamples/Private/ElectraTextureSampleAndroid.cpp
+++ b/Engine/Plugins/Media/ElectraUtil/Source/ElectraSamples/Private/ElectraTextureSampleAndroid.cpp
@@ -536,7 +536,7 @@ bool FElectraTextureSample::Convert(FTexture2DRHIRef& InDstTexture, const FConve
 
 	// Setup conversion from Rec2020 to current working color space
 	const UE::Color::FColorSpace& Working = UE::Color::FColorSpace::GetWorking();
-	FMatrix44f ColorSpaceMtx = UE::Color::Transpose<float>(Working.GetXYZToRgb()) * GetGamutToXYZMatrix();
+	FMatrix44f ColorSpaceMtx = FMatrix44f(Working.GetXYZToRgb().GetTransposed() * GetGamutToXYZMatrix());
 	if (GetEncodingType() == UE::Color::EEncoding::ST2084)
 	{
 		// Normalize output (e.g. 80 or 100 nits == 1.0)
diff --git a/Engine/Plugins/Media/ElectraUtil/Source/ElectraSamples/Private/ElectraTextureSampleBase.cpp b/Engine/Plugins/Media/ElectraUtil/Source/ElectraSamples/Private/ElectraTextureSampleBase.cpp
--- a/Engine/Plugins/Media/ElectraUtil/Source/ElectraSamples/Private/ElectraTextureSampleBase.cpp
+++ b/Engine/Plugins/Media/ElectraUtil/Source/ElectraSamples/Private/ElectraTextureSampleBase.cpp
@@ -56,6 +56,11 @@ void IElectraTextureSampleBase::Initialize(FVideoDecoderOutput* InVideoDecoderOu
 	FVector Off = FVector::Zero();
 	const FMatrix* Mtx = nullptr;
 
+	DisplayMasteringLuminanceMin = -1.0f;
+	DisplayMasteringLuminanceMax = -1.0f;
+	MaxCLL = 0;
+	MaxFALL = 0;
+
 	// Do we have specific HDR information, so we can assume a standard?
 	if (auto PinnedHDRInfo = HDRInfo.Pin())
 	{
@@ -65,40 +70,24 @@ void IElectraTextureSampleBase::Initialize(FVideoDecoderOutput* InVideoDecoderOu
 
 		if (auto ColorVolume = PinnedHDRInfo->GetMasteringDisplayColourVolume())
 		{
-			WhitePoint = FVector2d(ColorVolume->white_point_x, ColorVolume->white_point_y);
-			for (uint32 I = 0; I < 3; ++I)
-			{
-				DisplayPrimaries[I] = FVector2d(ColorVolume->display_primaries_x[I], ColorVolume->display_primaries_y[I]);
-			}
+			SampleColorSpace = UE::Color::FColorSpace(FVector2d(ColorVolume->display_primaries_x[0], ColorVolume->display_primaries_y[0]),
+													  FVector2d(ColorVolume->display_primaries_x[1], ColorVolume->display_primaries_y[1]),
+													  FVector2d(ColorVolume->display_primaries_x[2], ColorVolume->display_primaries_y[2]),
+													  FVector2d(ColorVolume->white_point_x, ColorVolume->white_point_y));
+
+			DisplayMasteringLuminanceMin = ColorVolume->min_display_mastering_luminance;
+			DisplayMasteringLuminanceMax = ColorVolume->max_display_mastering_luminance;
 		}
 		else
 		{
-			WhitePoint = FVector2d(UE::Color::GetWhitePoint(UE::Color::EWhitePoint::CIE1931_D65));
-
-			const FVector2d* DPs = ElectraColorimetryUtils::GetColorPrimaries(UE::Color::EColorSpace::Rec2020);
-			DisplayPrimaries[0] = DPs[0];
-			DisplayPrimaries[1] = DPs[1];
-			DisplayPrimaries[2] = DPs[2];
+			SampleColorSpace = UE::Color::FColorSpace(ElectraColorimetryUtils::TranslateMPEGColorPrimaries(ColorPrimaries));
 		}
 
-		switch (PinnedHDRInfo->GetHDRType())
+		if (auto ContentLightLevelInfo = PinnedHDRInfo->GetContentLightLevelInfo())
 		{
-		case IVideoDecoderHDRInformation::EType::PQ10:
-		case IVideoDecoderHDRInformation::EType::HDR10:
-			ColorEncoding = UE::Color::EEncoding::ST2084;
-			break;
-		case IVideoDecoderHDRInformation::EType::HLG10:
-			UE_LOG(LogElectraSamples, Warning, TEXT("Detected use of HLG EOTF in video material. Mapping to Rec709/sRGB as HLG is not supported."));
-			ColorEncoding = UE::Color::EEncoding::sRGB;	// using sRGB in place of HLG for now
-			break;
-		default:
-			check(!"Unknown HDR type!");
-			ColorEncoding = UE::Color::EEncoding::sRGB;
-			break;
+			MaxCLL = ContentLightLevelInfo->max_content_light_level;
+			MaxFALL =  ContentLightLevelInfo->max_pic_average_light_level;
 		}
-
-		ColorGamut = EDisplayColorGamut::Rec2020_D65;
-		Mtx = bFullRange ? &MediaShaders::YuvToRgbRec2020Unscaled : &MediaShaders::YuvToRgbRec2020Scaled;
 	}
 	else
 	{
@@ -106,36 +95,27 @@ void IElectraTextureSampleBase::Initialize(FVideoDecoderOutput* InVideoDecoderOu
 		// No HDR information present
 		//
 
-		const FVector2d* DPs = ElectraColorimetryUtils::GetColorPrimaries(ElectraColorimetryUtils::TranslateMPEGColorPrimaries(ColorPrimaries));
-		DisplayPrimaries[0] = DPs[0];
-		DisplayPrimaries[1] = DPs[1];
-		DisplayPrimaries[2] = DPs[2];
-
-		switch (ElectraColorimetryUtils::TranslateMPEGMatrixCoefficients(MatrixCoefficients))
-		{
-			case UE::Color::EColorSpace::None:	// ID (RGB)
-				// no conversion, data is RGB
-				break;
-			case UE::Color::EColorSpace::sRGB:
-				Mtx = bFullRange ? &MediaShaders::YuvToRgbRec709Unscaled : &MediaShaders::YuvToRgbRec709Scaled;
-				ColorGamut = EDisplayColorGamut::sRGB_D65;
-				break;
-			case UE::Color::EColorSpace::Rec2020:
-				Mtx = bFullRange ? &MediaShaders::YuvToRgbRec2020Unscaled : &MediaShaders::YuvToRgbRec2020Scaled;
-				ColorGamut = EDisplayColorGamut::Rec2020_D65;
-				break;
-			default:
-				check(!"*** Unexpected matrix coefficients!");
-				Mtx = bFullRange ? &MediaShaders::YuvToRgbRec709Unscaled : &MediaShaders::YuvToRgbRec709Scaled;
-				ColorGamut = EDisplayColorGamut::sRGB_D65;
-		}
-
-		ColorEncoding = ElectraColorimetryUtils::TranslateMPEGTransferCharacteristics(TransferCharacteristics);
+		SampleColorSpace = UE::Color::FColorSpace(ElectraColorimetryUtils::TranslateMPEGColorPrimaries(ColorPrimaries));
+	}
 
-		// No matter what: we will have the D65 white point (for now at least)
-		WhitePoint = FVector2d(UE::Color::GetWhitePoint(UE::Color::EWhitePoint::CIE1931_D65));
+	switch (ElectraColorimetryUtils::TranslateMPEGMatrixCoefficients(MatrixCoefficients))
+	{
+		case UE::Color::EColorSpace::None:	// ID (RGB)
+			// no conversion, data is RGB
+			break;
+		case UE::Color::EColorSpace::sRGB:
+			Mtx = bFullRange ? &MediaShaders::YuvToRgbRec709Unscaled : &MediaShaders::YuvToRgbRec709Scaled;
+			break;
+		case UE::Color::EColorSpace::Rec2020:
+			Mtx = bFullRange ? &MediaShaders::YuvToRgbRec2020Unscaled : &MediaShaders::YuvToRgbRec2020Scaled;
+			break;
+		default:
+			check(!"*** Unexpected matrix coefficients!");
+			Mtx = bFullRange ? &MediaShaders::YuvToRgbRec709Unscaled : &MediaShaders::YuvToRgbRec709Scaled;
 	}
 
+	ColorEncoding = ElectraColorimetryUtils::TranslateMPEGTransferCharacteristics(TransferCharacteristics);
+
 	if (Mtx)
 	{
 		// Select the offsets prior to YUV conversion needed per the incoming data
@@ -245,29 +225,29 @@ FMatrix44f IElectraTextureSampleBase::GetSampleToRGBMatrix() const
 	return SampleToRgbMtx;
 }
 
-FMatrix44f IElectraTextureSampleBase::GetGamutToXYZMatrix() const
+FMatrix44d IElectraTextureSampleBase::GetGamutToXYZMatrix() const
 {
-	return GamutToXYZMatrix(ColorGamut);
+	return SampleColorSpace.GetRgbToXYZ().GetTransposed();
 }
 
-FVector2f IElectraTextureSampleBase::GetWhitePoint() const
+FVector2d IElectraTextureSampleBase::GetWhitePoint() const
 {
-	return FVector2f(WhitePoint);
+	return SampleColorSpace.GetWhiteChromaticity();
 }
 
-FVector2f IElectraTextureSampleBase::GetDisplayPrimaryRed() const
+FVector2d IElectraTextureSampleBase::GetDisplayPrimaryRed() const
 {
-	return FVector2f(DisplayPrimaries[0]);
+	return SampleColorSpace.GetRedChromaticity();
 }
 
-FVector2f IElectraTextureSampleBase::GetDisplayPrimaryGreen() const
+FVector2d IElectraTextureSampleBase::GetDisplayPrimaryGreen() const
 {
-	return FVector2f(DisplayPrimaries[1]);
+	return SampleColorSpace.GetGreenChromaticity();
 }
 
-FVector2f IElectraTextureSampleBase::GetDisplayPrimaryBlue() const
+FVector2d IElectraTextureSampleBase::GetDisplayPrimaryBlue() const
 {
-	return FVector2f(DisplayPrimaries[2]);
+	return SampleColorSpace.GetBlueChromaticity();
 }
 
 UE::Color::EEncoding IElectraTextureSampleBase::GetEncodingType() const
@@ -275,4 +255,28 @@ UE::Color::EEncoding IElectraTextureSampleBase::GetEncodingType() const
 	return ColorEncoding;
 }
 
+bool IElectraTextureSampleBase::GetDisplayMasteringLuminance(float& OutMin, float& OutMax) const
+{
+	if (DisplayMasteringLuminanceMin < 0.0f && DisplayMasteringLuminanceMax < 0.0f)
+	{
+		return false;
+	}
+
+	OutMin = DisplayMasteringLuminanceMin;
+	OutMax = DisplayMasteringLuminanceMax;
+	return true;
+}
+
+bool IElectraTextureSampleBase::GetMaxLuminanceLevels(uint16& OutCLL, uint16& OutFALL) const
+{
+	if (MaxCLL == 0 && MaxFALL == 0)
+	{
+		return false;
+	}
+
+	OutCLL = MaxCLL;
+	OutFALL = MaxFALL;
+	return true;
+}
+
 #endif
diff --git a/Engine/Plugins/Media/ElectraUtil/Source/ElectraSamples/Public/Apple/ElectraTextureSample.h b/Engine/Plugins/Media/ElectraUtil/Source/ElectraSamples/Public/Apple/ElectraTextureSample.h
--- a/Engine/Plugins/Media/ElectraUtil/Source/ElectraSamples/Public/Apple/ElectraTextureSample.h
+++ b/Engine/Plugins/Media/ElectraUtil/Source/ElectraSamples/Public/Apple/ElectraTextureSample.h
@@ -73,7 +73,7 @@ public:
     ~FElectraMediaTexConvApple();
 
 #if WITH_ENGINE
-    void ConvertTexture(FTexture2DRHIRef & InDstTexture, CVImageBufferRef InImageBufferRef, bool bFullRange, EMediaTextureSampleFormat Format, const FMatrix44f& YUVMtx, const FMatrix44f& GamutToXYZMtx, UE::Color::EEncoding EncodingType, float NormalizationFactor);
+    void ConvertTexture(FTexture2DRHIRef & InDstTexture, CVImageBufferRef InImageBufferRef, bool bFullRange, EMediaTextureSampleFormat Format, const FMatrix44f& YUVMtx, const FMatrix44d& GamutToXYZMtx, UE::Color::EEncoding EncodingType, float NormalizationFactor);
 #endif
 
 private:
diff --git a/Engine/Plugins/Media/ElectraUtil/Source/ElectraSamples/Public/ElectraTextureSampleUtils.h b/Engine/Plugins/Media/ElectraUtil/Source/ElectraSamples/Public/ElectraTextureSampleUtils.h
--- a/Engine/Plugins/Media/ElectraUtil/Source/ElectraSamples/Public/ElectraTextureSampleUtils.h
+++ b/Engine/Plugins/Media/ElectraUtil/Source/ElectraSamples/Public/ElectraTextureSampleUtils.h
@@ -64,21 +64,4 @@ namespace ElectraColorimetryUtils
 		}
 		return UE::Color::EEncoding::sRGB;
 	}
-
-	const FVector2d* GetColorPrimaries(UE::Color::EColorSpace InColorSpace)
-	{
-		static const FVector2d DP_sRGB[3] = { {0.64, 0.33} , {0.30, 0.60}, {0.15, 0.06} };
-		static const FVector2d DP_Rec2020[3] = { {0.708, 0.292} , {0.170, 0.797}, {0.131, 0.046} };
-
-		switch (InColorSpace)
-		{
-		case UE::Color::EColorSpace::sRGB:
-			return DP_sRGB;
-		case UE::Color::EColorSpace::Rec2020:
-			return DP_Rec2020;
-		default:
-			check(!"*** Unexpected color primaries!");
-		}
-		return DP_sRGB;
-	}
 };
diff --git a/Engine/Plugins/Media/ElectraUtil/Source/ElectraSamples/Public/IElectraTextureSample.h b/Engine/Plugins/Media/ElectraUtil/Source/ElectraSamples/Public/IElectraTextureSample.h
--- a/Engine/Plugins/Media/ElectraUtil/Source/ElectraSamples/Public/IElectraTextureSample.h
+++ b/Engine/Plugins/Media/ElectraUtil/Source/ElectraSamples/Public/IElectraTextureSample.h
@@ -48,12 +48,14 @@ public:
 	virtual bool GetFullRange() const override;
 
 	virtual FMatrix44f GetSampleToRGBMatrix() const override;
-	virtual FMatrix44f GetGamutToXYZMatrix() const override;
-	virtual FVector2f GetWhitePoint() const override;
-	virtual FVector2f GetDisplayPrimaryRed() const override;
-	virtual FVector2f GetDisplayPrimaryGreen() const override;
-	virtual FVector2f GetDisplayPrimaryBlue() const override;
+	virtual FMatrix44d GetGamutToXYZMatrix() const override;
+	virtual FVector2d GetWhitePoint() const override;
+	virtual FVector2d GetDisplayPrimaryRed() const override;
+	virtual FVector2d GetDisplayPrimaryGreen() const override;
+	virtual FVector2d GetDisplayPrimaryBlue() const override;
 	virtual UE::Color::EEncoding GetEncodingType() const override;
+	virtual bool GetDisplayMasteringLuminance(float& OutMin, float& OutMax) const override;
+	virtual bool GetMaxLuminanceLevels(uint16& OutCLL, uint16& OutFALL) const override;
 
 protected:
 	virtual float GetSampleDataScale(bool b10Bit) const { return 1.0f; }
@@ -72,8 +74,10 @@ protected:
 	const FMatrix* YuvToRgbMtx;
 
 	/** Precomputed colorimetric data */
-	EDisplayColorGamut ColorGamut;
 	UE::Color::EEncoding ColorEncoding;
-	FVector2d WhitePoint;
-	FVector2d DisplayPrimaries[3];
+	UE::Color::FColorSpace SampleColorSpace;
+	float DisplayMasteringLuminanceMin;
+	float DisplayMasteringLuminanceMax;
+	uint16 MaxCLL;
+	uint16 MaxFALL;
 };
diff --git a/Engine/Plugins/Media/MediaIOFramework/Source/MediaIOCore/Private/Shared/MediaIOCoreTextureSampleBase.cpp b/Engine/Plugins/Media/MediaIOFramework/Source/MediaIOCore/Private/Shared/MediaIOCoreTextureSampleBase.cpp
--- a/Engine/Plugins/Media/MediaIOFramework/Source/MediaIOCore/Private/Shared/MediaIOCoreTextureSampleBase.cpp
+++ b/Engine/Plugins/Media/MediaIOFramework/Source/MediaIOCore/Private/Shared/MediaIOCoreTextureSampleBase.cpp
@@ -339,29 +339,29 @@ bool FMediaIOCoreTextureSampleBase::IsOutputSrgb() const
 	return Encoding == UE::Color::EEncoding::sRGB;
 }
 
-FMatrix44f FMediaIOCoreTextureSampleBase::GetGamutToXYZMatrix() const
+FMatrix44d FMediaIOCoreTextureSampleBase::GetGamutToXYZMatrix() const
 {
-	return FMatrix44f(ColorSpaceStruct.GetRgbToXYZ().GetTransposed());
+	return ColorSpaceStruct.GetRgbToXYZ().GetTransposed();
 }
 
-FVector2f FMediaIOCoreTextureSampleBase::GetWhitePoint() const
+FVector2d FMediaIOCoreTextureSampleBase::GetWhitePoint() const
 {
-	return FVector2f(ColorSpaceStruct.GetWhiteChromaticity());
+	return ColorSpaceStruct.GetWhiteChromaticity();
 }
 
-FVector2f FMediaIOCoreTextureSampleBase::GetDisplayPrimaryRed() const
+FVector2d FMediaIOCoreTextureSampleBase::GetDisplayPrimaryRed() const
 {
-	return FVector2f(ColorSpaceStruct.GetRedChromaticity());
+	return ColorSpaceStruct.GetRedChromaticity();
 }
 
-FVector2f FMediaIOCoreTextureSampleBase::GetDisplayPrimaryGreen() const
+FVector2d FMediaIOCoreTextureSampleBase::GetDisplayPrimaryGreen() const
 {
-	return FVector2f(ColorSpaceStruct.GetGreenChromaticity());
+	return ColorSpaceStruct.GetGreenChromaticity();
 } 
 
-FVector2f FMediaIOCoreTextureSampleBase::GetDisplayPrimaryBlue() const
+FVector2d FMediaIOCoreTextureSampleBase::GetDisplayPrimaryBlue() const
 {
-	return FVector2f(ColorSpaceStruct.GetBlueChromaticity());
+	return ColorSpaceStruct.GetBlueChromaticity();
 }
 
 UE::Color::EEncoding FMediaIOCoreTextureSampleBase::GetEncodingType() const
diff --git a/Engine/Plugins/Media/MediaIOFramework/Source/MediaIOCore/Public/MediaIOCoreTextureSampleBase.h b/Engine/Plugins/Media/MediaIOFramework/Source/MediaIOCore/Public/MediaIOCoreTextureSampleBase.h
--- a/Engine/Plugins/Media/MediaIOFramework/Source/MediaIOCore/Public/MediaIOCoreTextureSampleBase.h
+++ b/Engine/Plugins/Media/MediaIOFramework/Source/MediaIOCore/Public/MediaIOCoreTextureSampleBase.h
@@ -315,11 +315,11 @@ public:
 
 	virtual bool IsOutputSrgb() const override;
 	
-	virtual FMatrix44f GetGamutToXYZMatrix() const override;
-	virtual FVector2f GetWhitePoint() const override;
-	virtual FVector2f GetDisplayPrimaryRed() const override;
-	virtual FVector2f GetDisplayPrimaryGreen() const override;
-	virtual FVector2f GetDisplayPrimaryBlue() const override;
+	virtual FMatrix44d GetGamutToXYZMatrix() const override;
+	virtual FVector2d GetWhitePoint() const override;
+	virtual FVector2d GetDisplayPrimaryRed() const override;
+	virtual FVector2d GetDisplayPrimaryGreen() const override;
+	virtual FVector2d GetDisplayPrimaryBlue() const override;
 	virtual UE::Color::EEncoding GetEncodingType() const override;
 	virtual float GetHDRNitsNormalizationFactor() const override;
 
diff --git a/Engine/Source/Runtime/Media/Public/IMediaTextureSample.h b/Engine/Source/Runtime/Media/Public/IMediaTextureSample.h
--- a/Engine/Source/Runtime/Media/Public/IMediaTextureSample.h
+++ b/Engine/Source/Runtime/Media/Public/IMediaTextureSample.h
@@ -410,9 +410,9 @@ public:
 	 * 
 	 * @return Conversion Matrix
 	 */
-	virtual FMatrix44f GetGamutToXYZMatrix() const
+	virtual FMatrix44d GetGamutToXYZMatrix() const
 	{
-		return GamutToXYZMatrix(EDisplayColorGamut::sRGB_D65);
+		return FMatrix44d(GamutToXYZMatrix(EDisplayColorGamut::sRGB_D65));
 	}
 
 	/**
@@ -420,24 +420,24 @@ public:
 	 * 
 	 * @return White point
 	 */
-	virtual FVector2f GetWhitePoint() const
+	virtual FVector2d GetWhitePoint() const
 	{
-		return FVector2f(UE::Color::GetWhitePoint(UE::Color::EWhitePoint::CIE1931_D65));
+		return UE::Color::GetWhitePoint(UE::Color::EWhitePoint::CIE1931_D65);
 	}
 
-	virtual FVector2f GetDisplayPrimaryRed() const
+	virtual FVector2d GetDisplayPrimaryRed() const
 	{
-		return FVector2f(0.64, 0.33);
+		return FVector2d(0.64, 0.33);
 	}
 
-	virtual FVector2f GetDisplayPrimaryGreen() const
+	virtual FVector2d GetDisplayPrimaryGreen() const
 	{
-		return FVector2f(0.30, 0.60);
+		return FVector2d(0.30, 0.60);
 	}
 
-	virtual FVector2f GetDisplayPrimaryBlue() const
+	virtual FVector2d GetDisplayPrimaryBlue() const
 	{
-		return FVector2f(0.15, 0.06);
+		return FVector2d(0.15, 0.06);
 	}
 
 	/**
@@ -448,11 +448,30 @@ public:
 		return IsOutputSrgb() ? UE::Color::EEncoding::sRGB : UE::Color::EEncoding::Linear;
 	}
 
+	/**
+	 * Get factor to normalize data from nits to scene color values
+	 */
 	virtual float GetHDRNitsNormalizationFactor() const
 	{
 		return (GetEncodingType() == UE::Color::EEncoding::sRGB || GetEncodingType() == UE::Color::EEncoding::Linear) ? 1.0f : kMediaSample_HDR_NitsNormalizationFactor;
 	}
 
+	/**
+	 * Get display mastering luminance information
+	 */
+	virtual bool GetDisplayMasteringLuminance(float& OutMin, float& OutMax) const
+	{
+		return false;
+	}
+
+	/**
+	 * Get maximum luminance information
+	 */
+	virtual bool GetMaxLuminanceLevels(uint16& OutCLL, uint16& OutFALL) const
+	{
+		return false;
+	}
+
 	/**
 	 * Reset sample to empty state
 	 */
diff --git a/Engine/Source/Runtime/MediaAssets/Private/Misc/MediaTextureResource.cpp b/Engine/Source/Runtime/MediaAssets/Private/Misc/MediaTextureResource.cpp
--- a/Engine/Source/Runtime/MediaAssets/Private/Misc/MediaTextureResource.cpp
+++ b/Engine/Source/Runtime/MediaAssets/Private/Misc/MediaTextureResource.cpp
@@ -887,7 +887,7 @@ bool FMediaTextureResource::RequiresConversion(const TSharedPtr<IMediaTextureSam
 	//
 	// Some reasons this could trigger:
 	//	- Cropping
-	//  - NV12's 1.5x heigth factor -> but the format would need conversion to RGB anyway
+	//  - NV12's 1.5x height factor -> but the format would need conversion to RGB anyway
 	//	- Various YCbCr formats will be stored with different widths as the final RGB output data -> a conversion is needed anyway
 	//
 	if (Sample->GetDim() != Sample->GetOutputDim())
@@ -898,10 +898,10 @@ bool FMediaTextureResource::RequiresConversion(const TSharedPtr<IMediaTextureSam
 	// Color space different?
 	const UE::Color::FColorSpace& Working = OverrideColorSpace.IsValid() ? *OverrideColorSpace : UE::Color::FColorSpace::GetWorking();
 	const float Tollerance = 1.e-7f;
-	if (!Sample->GetDisplayPrimaryRed().Equals(FVector2f(Working.GetRedChromaticity()), Tollerance) ||
-		!Sample->GetDisplayPrimaryGreen().Equals(FVector2f(Working.GetGreenChromaticity()), Tollerance) ||
-		!Sample->GetDisplayPrimaryBlue().Equals(FVector2f(Working.GetBlueChromaticity()), Tollerance) ||
-		!Sample->GetWhitePoint().Equals(FVector2f(Working.GetWhiteChromaticity()), Tollerance))
+	if (!Sample->GetDisplayPrimaryRed().Equals(Working.GetRedChromaticity(), Tollerance) ||
+		!Sample->GetDisplayPrimaryGreen().Equals(Working.GetGreenChromaticity(), Tollerance) ||
+		!Sample->GetDisplayPrimaryBlue().Equals(Working.GetBlueChromaticity(), Tollerance) ||
+		!Sample->GetWhitePoint().Equals(Working.GetWhiteChromaticity(), Tollerance))
 	{
 		// Yes! We need to convert...
 		return true;
@@ -1028,7 +1028,7 @@ void FMediaTextureResource::GetColorSpaceConversionMatrixForSample(const TShared
 	}
 	else
 	{
-		ColorSpaceMtx = UE::Color::Transpose<float>(Working.GetXYZToRgb()) * Sample->GetGamutToXYZMatrix();
+		ColorSpaceMtx = FMatrix44f(Working.GetXYZToRgb().GetTransposed() * Sample->GetGamutToXYZMatrix());
 	}
 	
 	float NF = Sample->GetHDRNitsNormalizationFactor();
