ElectraUtil: reserving string size to reduce or eliminate reallocations during URL parsing

#rnx

diff --git a/Engine/Plugins/Media/ElectraUtil/Source/ElectraBase/Private/URLParser.cpp b/Engine/Plugins/Media/ElectraUtil/Source/ElectraBase/Private/URLParser.cpp
--- a/Engine/Plugins/Media/ElectraUtil/Source/ElectraBase/Private/URLParser.cpp
+++ b/Engine/Plugins/Media/ElectraUtil/Source/ElectraBase/Private/URLParser.cpp
@@ -5,6 +5,30 @@
 
 namespace Electra
 {
+	namespace URLComponents
+	{
+		static const FString KeepCharsPath(TEXT("/"));
+		static const FString KeepCharsFragment(TEXT("!$&'()*+,;=;@/?"));
+
+		static const FString URN(TEXT("urn:"));
+		static const FString Slash(TEXT("/"));
+		static const FString DotSlash(TEXT("./"));
+		static const FString SlashSlash(TEXT("//"));
+		static const FString SchemeFile(TEXT("file"));
+		static const FString SchemeHttp(TEXT("http"));
+		static const FString SchemeHttps(TEXT("https"));
+		static const FString Colon(TEXT(":"));
+		static const FString QuestionMark(TEXT("?"));
+		static const FString HashTag(TEXT("#"));
+		static const FString Ampersand(TEXT("&"));
+		static const FString At(TEXT("@"));
+		static const FString BracketOpen(TEXT("["));
+		static const FString BracketClose(TEXT("]"));
+		static const FString Dot(TEXT("."));
+		static const FString DotDot(TEXT(".."));
+		static const FString Port80(TEXT("80"));
+		static const FString Port443(TEXT("443"));
+	}
 
 	static inline void _SwapStrings(FString& a, FString& b)
 	{
@@ -43,7 +67,7 @@ namespace Electra
 			return true;
 		}
 		// Not handling URNs here.
-		if (InURL.StartsWith(TEXT("urn:")))
+		if (InURL.StartsWith(URLComponents::URN))
 		{
 			return false;
 		}
@@ -56,7 +80,7 @@ namespace Electra
 			   scheme     authority       path        query   fragment
 		*/
 		StringHelpers::FStringIterator it(InURL);
-		if (InURL.StartsWith(TEXT("//")))
+		if (InURL.StartsWith(URLComponents::SlashSlash))
 		{
 			it += 2;
 			if (ParseAuthority(it))
@@ -75,6 +99,7 @@ namespace Electra
 			// If we get a URL without a scheme we will trip over the colon separating host and port, mistaking it for
 			// the scheme delimiter. For our purposes we only handle URLs that have a scheme when they have an authority!
 			FString Component;
+			Component.Reserve(InURL.Len());
 			// Parse out the first component up to the delimiting colon, path, query or fragment.
 			while(it && !IsColonSeparator(*it) && !IsQueryOrFragmentSeparator(*it) && !IsPathSeparator(*it))
 			{
@@ -90,7 +115,7 @@ namespace Electra
 					return false;
 				}
 				// Assume we parsed the scheme.
-				Scheme = Component;
+				Scheme = MoveTemp(Component);
 				// Does an absolute path ('/') or an authority ('//') follow?
 				if (IsPathSeparator(*it))
 				{
@@ -158,85 +183,81 @@ namespace Electra
 
 	FString FURL_RFC3986::Get(bool bIncludeQuery, bool bIncludeFragment)
 	{
-		static const FString KeepCharsPath(TEXT("/"));
-		static const FString KeepCharsFragment(TEXT("!$&'()*+,;=;@/?"));
-
 		FString URL;
+		URL.Reserve(Scheme.Len() + UserInfo.Len() + Host.Len() + Port.Len() + Path.Len() + Query.Len() + Fragment.Len() + 32);
 		if (IsAbsolute())
 		{
-			bool bIsFile = Scheme.Equals(TEXT("file"));
+			bool bIsFile = Scheme.Equals(URLComponents::SchemeFile);
 			FString Authority = GetAuthority();
 			URL = Scheme;
-			URL += TEXT(":");
+			URL += URLComponents::Colon;
 			if (Authority.Len() || bIsFile)
 			{
-				URL += TEXT("//");
+				URL += URLComponents::SlashSlash;
 				URL += Authority;
 			}
 			if (Path.Len())
 			{
 				if (Authority.Len() && !IsPathSeparator(Path[0]))
 				{
-					URL += TEXT("/");
+					URL += URLComponents::Slash;
 				}
-				UrlEncode(URL, Path, KeepCharsPath);
+				UrlEncode(URL, Path, URLComponents::KeepCharsPath);
 			}
 			else if ((Query.Len() && bIncludeQuery) || (Fragment.Len() && bIncludeFragment))
 			{
-				URL += TEXT("/");
+				URL += URLComponents::Slash;
 			}
 		}
 		else
 		{
-			UrlEncode(URL, Path, KeepCharsPath);
+			UrlEncode(URL, Path, URLComponents::KeepCharsPath);
 		}
 		if (Query.Len() && bIncludeQuery)
 		{
-			URL += TEXT("?");
+			URL += URLComponents::QuestionMark;
 			URL += Query;
 		}
 		if (Fragment.Len() && bIncludeFragment)
 		{
-			URL += TEXT("#");
-			UrlEncode(URL, Fragment, KeepCharsFragment);
+			URL += URLComponents::HashTag;
+			UrlEncode(URL, Fragment, URLComponents::KeepCharsFragment);
 		}
 		return URL;
 	}
 
 	FString FURL_RFC3986::GetPath(bool bIncludeQuery, bool bIncludeFragment)
 	{
-		static const FString KeepCharsPath(TEXT("/"));
-		static const FString KeepCharsFragment(TEXT("!$&'()*+,;=;@/?"));
-
 		FString URL;
+		URL.Reserve(Path.Len() + Query.Len() + Fragment.Len() + 32);
 		if (IsAbsolute())
 		{
 			if (Path.Len())
 			{
 				if (GetAuthority().Len() && !IsPathSeparator(Path[0]))
 				{
-					URL += TEXT("/");
+					URL += URLComponents::Slash;
 				}
-				UrlEncode(URL, Path, KeepCharsPath);
+				UrlEncode(URL, Path, URLComponents::KeepCharsPath);
 			}
 			else if ((Query.Len() && bIncludeQuery) || (Fragment.Len() && bIncludeFragment))
 			{
-				URL += TEXT("/");
+				URL += URLComponents::Slash;
 			}
 		}
 		else
 		{
-			UrlEncode(URL, Path, KeepCharsPath);
+			UrlEncode(URL, Path, URLComponents::KeepCharsPath);
 		}
 		if (Query.Len() && bIncludeQuery)
 		{
-			URL += TEXT("?");
+			URL += URLComponents::QuestionMark;
 			URL += Query;
 		}
 		if (Fragment.Len() && bIncludeFragment)
 		{
-			URL += TEXT("#");
-			UrlEncode(URL, Fragment, KeepCharsFragment);
+			URL += URLComponents::HashTag;
+			UrlEncode(URL, Fragment, URLComponents::KeepCharsFragment);
 		}
 		return URL;
 	}
@@ -244,7 +265,7 @@ namespace Electra
 	void FURL_RFC3986::GetQueryParams(TArray<FQueryParam>& OutQueryParams, const FString& InQueryParameters, bool bPerformUrlDecoding, bool bSameNameReplacesValue)
 	{
 		TArray<FString> ValuePairs;
-		InQueryParameters.ParseIntoArray(ValuePairs, TEXT("&"), true);
+		InQueryParameters.ParseIntoArray(ValuePairs, *URLComponents::Ampersand, true);
 		for(int32 i=0; i<ValuePairs.Num(); ++i)
 		{
 			int32 EqPos = 0;
@@ -295,7 +316,13 @@ namespace Electra
 	void FURL_RFC3986::SetQueryParams(const TArray<FQueryParam>& InQueryParams)
 	{
 		Query.Empty();
-		for (int i = 0; i < InQueryParams.Num(); ++i)
+		int32 Len=0;
+		for(int32 i=0; i<InQueryParams.Num(); ++i)
+		{
+			Len += InQueryParams[i].Name.Len() + InQueryParams[i].Value.Len() + 2;
+		}
+		Query.Reserve(Len);
+		for(int32 i=0; i<InQueryParams.Num(); ++i)
 		{
 			Query += InQueryParams[i].Name + '=' + InQueryParams[i].Value;
 			if (i != InQueryParams.Num()-1)
@@ -362,12 +389,13 @@ namespace Electra
 	FString FURL_RFC3986::GetAuthority() const
 	{
 		FString Authority;
+		Authority.Reserve(Scheme.Len() + UserInfo.Len() + Host.Len() + Port.Len() + 32);
 		if (UserInfo.Len())
 		{
 			Authority += UserInfo;
-			Authority += TEXT("@");
+			Authority += URLComponents::At;
 		}
-		if (Scheme.Equals(TEXT("file")))
+		if (Scheme.Equals(URLComponents::SchemeFile))
 		{
 			Authority += Host;
 		}
@@ -377,9 +405,9 @@ namespace Electra
 			int32 DummyPos = 0;
 			if (Host.FindChar(TCHAR(':'), DummyPos))
 			{
-				Authority += TEXT("[");
+				Authority += URLComponents::BracketOpen;
 				Authority += Host;
-				Authority += TEXT("]");
+				Authority += URLComponents::BracketClose;
 			}
 			else
 			{
@@ -388,7 +416,7 @@ namespace Electra
 			// Need to append a port?
 			if (Port.Len())
 			{
-				Authority += TEXT(":");
+				Authority += URLComponents::Colon;
 				Authority += Port;
 			}
 		}
@@ -416,6 +444,7 @@ namespace Electra
 	{
 		UserInfo.Empty();
 		FString Component;
+		Component.Reserve(it.GetRemainingLength());
 		while(it && !IsPathSeparator(*it) && !IsQueryOrFragmentSeparator(*it))
 		{
 			// If there is a user-info delimiter?
@@ -442,7 +471,8 @@ namespace Electra
 			return true;
 		}
 		Host.Empty();
-		if (Scheme.Equals(TEXT("file")))
+		Host.Reserve(it.GetRemainingLength());
+		if (Scheme.Equals(URLComponents::SchemeFile))
 		{
 			// For file:// scheme we have to consider Windows drive letters with a colon, eg. file://D:/
 			// We must not stop parsing at the colon since it will not indicate a port for the file scheme anyway.
@@ -479,6 +509,7 @@ namespace Electra
 			{
 				++it;
 				Port.Empty();
+				Port.Reserve(16);
 				while(it)
 				{
 					Port += *it++;
@@ -517,6 +548,7 @@ namespace Electra
 	bool FURL_RFC3986::ParsePath(StringHelpers::FStringIterator& it)
 	{
 		Path.Empty();
+		Path.Reserve(it.GetRemainingLength());
 		while(it && !IsQueryOrFragmentSeparator(*it))
 		{
 			Path += *it++;
@@ -534,6 +566,7 @@ namespace Electra
 	bool FURL_RFC3986::ParseQuery(StringHelpers::FStringIterator& it)
 	{
 		Query.Empty();
+		Query.Reserve(it.GetRemainingLength());
 		// Query is all up to the end or the start of the fragment.
 		while(it && !IsFragmentSeparator(*it))
 		{
@@ -625,6 +658,7 @@ namespace Electra
 			return (In >= ANSICHAR('0') && In <= ANSICHAR('9')) || (In >= ANSICHAR('a') && In <= ANSICHAR('f')) || (In >= ANSICHAR('A') && In <= ANSICHAR('F'));
 		};
 
+		OutResult.Reserve(InUrlToEncode.Len() * 3);	// assume we need to encode every character
 		FTCHARToUTF8 UTF8String(*InUrlToEncode);
 		for(const ANSICHAR* InUTF8Bytes=UTF8String.Get(); *InUTF8Bytes; ++InUTF8Bytes)
 		{
@@ -655,13 +689,13 @@ namespace Electra
 
 	FString FURL_RFC3986::GetStandardPortForScheme(const FString& InScheme, bool bIgnoreCase)
 	{
-		if (InScheme.Equals(TEXT("http"), bIgnoreCase ? ESearchCase::IgnoreCase : ESearchCase::CaseSensitive))
+		if (InScheme.Equals(URLComponents::SchemeHttp, bIgnoreCase ? ESearchCase::IgnoreCase : ESearchCase::CaseSensitive))
 		{
-			return FString(TEXT("80"));
+			return URLComponents::Port80;
 		}
-		else if (InScheme.Equals(TEXT("https"), bIgnoreCase ? ESearchCase::IgnoreCase : ESearchCase::CaseSensitive))
+		else if (InScheme.Equals(URLComponents::SchemeHttps, bIgnoreCase ? ESearchCase::IgnoreCase : ESearchCase::CaseSensitive))
 		{
-			return FString(TEXT("443"));
+			return URLComponents::Port443;
 		}
 		return FString();
 	}
@@ -670,7 +704,7 @@ namespace Electra
 	{
 		TArray<FString> Components;
 		// Split on '/', ignoring resulting empty parts (eg. "a//b" gives ["a", "b"] only instead of ["a", "", "b"] !!
-		InPath.ParseIntoArray(Components, TEXT("/"), true);
+		InPath.ParseIntoArray(Components, *URLComponents::Slash, true);
 		OutPathComponents.Append(Components);
 	}
 
@@ -751,6 +785,12 @@ namespace Electra
 	void FURL_RFC3986::BuildPathFromSegments(const TArray<FString>& Components, bool bAddLeadingSlash, bool bAddTrailingSlash)
 	{
 		Path.Empty();
+		int32 ComponentLen = 0;
+		for(int32 i=0, iMax=Components.Num(); i<iMax; ++i)
+		{
+			ComponentLen += Components[i].Len() + 1;
+		}
+		Path.Reserve(ComponentLen + 8);
 		int32 DummyPos = 0;
 		for(int32 i=0, iMax=Components.Num(); i<iMax; ++i)
 		{
@@ -758,7 +798,7 @@ namespace Electra
 			{
 				if (bAddLeadingSlash)
 				{
-					Path = TEXT("/");
+					Path = URLComponents::Slash;
 				}
 				/*
 					As per RFC 3986 Section 4.2 Relative Reference:
@@ -770,7 +810,7 @@ namespace Electra
 				*/
 				else if (Scheme.IsEmpty() && Components[0].FindChar(TCHAR(':'), DummyPos))
 				{
-					Path = TEXT("./");
+					Path = URLComponents::DotSlash;
 				}
 			}
 			else
@@ -790,7 +830,8 @@ namespace Electra
 		// RFC 3986 Section 5.2.3. Merge Paths
 		if (Host.Len() && Path.IsEmpty())
 		{
-			Path = TEXT("/");
+			Path.Reserve(InPathToMerge.Len() + 1);
+			Path = URLComponents::Slash;
 			Path += InPathToMerge;
 		}
 		else
@@ -819,14 +860,14 @@ namespace Electra
 			TArray<FString> NewSegments;
 			// Remember the current path starting or ending in a slash.
 			bool bHasLeadingSlash = IsPathSeparator(Path[0]);
-			bool bHasTrailingSlash = Path.EndsWith(TEXT("/"));
+			bool bHasTrailingSlash = Path.EndsWith(URLComponents::Slash);
 			// Split the path into its segments
 			TArray<FString> Segments;
 			GetPathComponents(Segments);
 			for(int32 i=0; i<Segments.Num(); ++i)
 			{
 				// For every ".." go up one level.
-				if (Segments[i].Equals(TEXT("..")))
+				if (Segments[i].Equals(URLComponents::DotDot))
 				{
 					if (NewSegments.Num())
 					{
@@ -834,7 +875,7 @@ namespace Electra
 					}
 				}
 				// Add non- "." segments.
-				else if (!Segments[i].Equals(TEXT(".")))
+				else if (!Segments[i].Equals(URLComponents::Dot))
 				{
 					NewSegments.Add(Segments[i]);
 				}
diff --git a/Engine/Plugins/Media/ElectraUtil/Source/ElectraBase/Public/Utilities/StringHelpers.h b/Engine/Plugins/Media/ElectraUtil/Source/ElectraBase/Public/Utilities/StringHelpers.h
--- a/Engine/Plugins/Media/ElectraUtil/Source/ElectraBase/Public/Utilities/StringHelpers.h
+++ b/Engine/Plugins/Media/ElectraUtil/Source/ElectraBase/Public/Utilities/StringHelpers.h
@@ -115,6 +115,11 @@ namespace Electra
 				return &StringToIterate[Index];
 			}
 
+			int32 GetRemainingLength() const
+			{
+				return StringToIterate.Len() - Index;
+			}
+
 			/** Returns an index to the current element. */
 			int32 GetIndex() const
 			{
