Fixing up texture sample setup from colorimetry information returned by Electra decoders to correctly reflect encoding vs. mastering display information

diff --git a/Engine/Plugins/Media/ElectraUtil/Source/ElectraSamples/Private/ElectraTextureSampleBase.cpp b/Engine/Plugins/Media/ElectraUtil/Source/ElectraSamples/Private/ElectraTextureSampleBase.cpp
--- a/Engine/Plugins/Media/ElectraUtil/Source/ElectraSamples/Private/ElectraTextureSampleBase.cpp
+++ b/Engine/Plugins/Media/ElectraUtil/Source/ElectraSamples/Private/ElectraTextureSampleBase.cpp
@@ -56,6 +56,8 @@ void IElectraTextureSampleBase::Initialize(FVideoDecoderOutput* InVideoDecoderOu
 	FVector Off = FVector::Zero();
 	const FMatrix* Mtx = nullptr;
 
+	// Defaults in case no HDR info is present
+	bDisplayColorSpaceValid = false;
 	DisplayMasteringLuminanceMin = -1.0f;
 	DisplayMasteringLuminanceMax = -1.0f;
 	MaxCLL = 0;
@@ -68,36 +70,40 @@ void IElectraTextureSampleBase::Initialize(FVideoDecoderOutput* InVideoDecoderOu
 		// HDR information present
 		//
 
+		// Mastering display info...
 		if (auto ColorVolume = PinnedHDRInfo->GetMasteringDisplayColourVolume())
 		{
-			SampleColorSpace = UE::Color::FColorSpace(FVector2d(ColorVolume->display_primaries_x[0], ColorVolume->display_primaries_y[0]),
-													  FVector2d(ColorVolume->display_primaries_x[1], ColorVolume->display_primaries_y[1]),
-													  FVector2d(ColorVolume->display_primaries_x[2], ColorVolume->display_primaries_y[2]),
-													  FVector2d(ColorVolume->white_point_x, ColorVolume->white_point_y));
+			// A few sanity checks on the primaries coordinates (by no means exhaustive, but it should catch a fair share of oddities)
+			if (ColorVolume->display_primaries_x[0] > FMath::Max(ColorVolume->display_primaries_x[1], ColorVolume->display_primaries_x[2]) &&	// Red has largest X
+				ColorVolume->display_primaries_y[1] > FMath::Max(ColorVolume->display_primaries_y[0], ColorVolume->display_primaries_y[2]) &&	// Green has largest Y
+				ColorVolume->display_primaries_x[2] <= ColorVolume->display_primaries_x[0] &&													// Blue's X is smaller than Red's
+				ColorVolume->display_primaries_y[2] <= ColorVolume->display_primaries_y[0] &&													// Blue's Y is smaller than Red's
+				ColorVolume->display_primaries_x[2] <= ColorVolume->display_primaries_x[1] && 													// Blue's X is smaller or same than Green's
+				ColorVolume->display_primaries_x[1] <= ColorVolume->display_primaries_x[0]) 													// Red's X is greater or same than Green's
+			{
+				DisplayColorSpace = UE::Color::FColorSpace(FVector2d(ColorVolume->display_primaries_x[0], ColorVolume->display_primaries_y[0]),
+														   FVector2d(ColorVolume->display_primaries_x[1], ColorVolume->display_primaries_y[1]),
+														   FVector2d(ColorVolume->display_primaries_x[2], ColorVolume->display_primaries_y[2]),
+														   FVector2d(ColorVolume->white_point_x, ColorVolume->white_point_y));
+				bDisplayColorSpaceValid = true;
+			}
 
 			DisplayMasteringLuminanceMin = ColorVolume->min_display_mastering_luminance;
 			DisplayMasteringLuminanceMax = ColorVolume->max_display_mastering_luminance;
 		}
-		else
-		{
-			SampleColorSpace = UE::Color::FColorSpace(ElectraColorimetryUtils::TranslateMPEGColorPrimaries(ColorPrimaries));
-		}
 
+		// Content light level info...
 		if (auto ContentLightLevelInfo = PinnedHDRInfo->GetContentLightLevelInfo())
 		{
 			MaxCLL = ContentLightLevelInfo->max_content_light_level;
 			MaxFALL =  ContentLightLevelInfo->max_pic_average_light_level;
 		}
 	}
-	else
-	{
-		//
-		// No HDR information present
-		//
 
-		SampleColorSpace = UE::Color::FColorSpace(ElectraColorimetryUtils::TranslateMPEGColorPrimaries(ColorPrimaries));
-	}
+	// The sample color space is always defined by the color primaries value
+	SampleColorSpace = UE::Color::FColorSpace(ElectraColorimetryUtils::TranslateMPEGColorPrimaries(ColorPrimaries));
 
+	// Select the YUV-RGB conversion matrix to use
 	switch (ElectraColorimetryUtils::TranslateMPEGMatrixCoefficients(MatrixCoefficients))
 	{
 		case UE::Color::EColorSpace::None:	// ID (RGB)
@@ -114,6 +120,7 @@ void IElectraTextureSampleBase::Initialize(FVideoDecoderOutput* InVideoDecoderOu
 			Mtx = bFullRange ? &MediaShaders::YuvToRgbRec709Unscaled : &MediaShaders::YuvToRgbRec709Scaled;
 	}
 
+	// Get color encoding (sRGB, linear, PQ, HLG...)
 	ColorEncoding = ElectraColorimetryUtils::TranslateMPEGTransferCharacteristics(TransferCharacteristics);
 
 	if (Mtx)
@@ -232,22 +239,22 @@ FMatrix44d IElectraTextureSampleBase::GetGamutToXYZMatrix() const
 
 FVector2d IElectraTextureSampleBase::GetWhitePoint() const
 {
-	return SampleColorSpace.GetWhiteChromaticity();
+	return bDisplayColorSpaceValid ? DisplayColorSpace.GetWhiteChromaticity() : FVector2d(-1.0, -1.0);
 }
 
 FVector2d IElectraTextureSampleBase::GetDisplayPrimaryRed() const
 {
-	return SampleColorSpace.GetRedChromaticity();
+	return bDisplayColorSpaceValid ? DisplayColorSpace.GetRedChromaticity() : FVector2d(-1.0, -1.0);
 }
 
 FVector2d IElectraTextureSampleBase::GetDisplayPrimaryGreen() const
 {
-	return SampleColorSpace.GetGreenChromaticity();
+	return bDisplayColorSpaceValid ? DisplayColorSpace.GetGreenChromaticity() : FVector2d(-1.0, -1.0);
 }
 
 FVector2d IElectraTextureSampleBase::GetDisplayPrimaryBlue() const
 {
-	return SampleColorSpace.GetBlueChromaticity();
+	return bDisplayColorSpaceValid ? DisplayColorSpace.GetBlueChromaticity() : FVector2d(-1.0, -1.0);
 }
 
 UE::Color::EEncoding IElectraTextureSampleBase::GetEncodingType() const
diff --git a/Engine/Plugins/Media/ElectraUtil/Source/ElectraSamples/Public/IElectraTextureSample.h b/Engine/Plugins/Media/ElectraUtil/Source/ElectraSamples/Public/IElectraTextureSample.h
--- a/Engine/Plugins/Media/ElectraUtil/Source/ElectraSamples/Public/IElectraTextureSample.h
+++ b/Engine/Plugins/Media/ElectraUtil/Source/ElectraSamples/Public/IElectraTextureSample.h
@@ -76,6 +76,8 @@ protected:
 	/** Precomputed colorimetric data */
 	UE::Color::EEncoding ColorEncoding;
 	UE::Color::FColorSpace SampleColorSpace;
+	UE::Color::FColorSpace DisplayColorSpace;
+	bool bDisplayColorSpaceValid;
 	float DisplayMasteringLuminanceMin;
 	float DisplayMasteringLuminanceMax;
 	uint16 MaxCLL;
