Fixing possible deadlock during texture creation for video playback
#tests tests on local machine using testbed application

diff --git a/Engine/Plugins/Media/ElectraUtil/Source/ElectraSamples/Private/ElectraSamplesTextureManager.cpp b/Engine/Plugins/Media/ElectraUtil/Source/ElectraSamples/Private/ElectraSamplesTextureManager.cpp
--- a/Engine/Plugins/Media/ElectraUtil/Source/ElectraSamples/Private/ElectraSamplesTextureManager.cpp
+++ b/Engine/Plugins/Media/ElectraUtil/Source/ElectraSamples/Private/ElectraSamplesTextureManager.cpp
@@ -30,10 +30,11 @@ void FElectraSamplesTextureManager::CleanupMap()
 
 TSharedPtr<IVideoDecoderTexture, ESPMode::ThreadSafe> FElectraSamplesTextureManager::CreateTexture(const FIntPoint& Dim, EPixelFormat Fmt)
 {
-	FScopeLock Lock(&AccessCS);
+	// note: we create the texture outside the CS to avoid any danger of deadlocks
+	TSharedPtr<IVideoDecoderTexture, ESPMode::ThreadSafe> NewTexture = PlatformCreateTexture(Dim, Fmt);
 
+	FScopeLock Lock(&AccessCS);
 	CleanupMap();
-	TSharedPtr<IVideoDecoderTexture, ESPMode::ThreadSafe> NewTexture = PlatformCreateTexture(Dim, Fmt);
 	if (NewTexture.IsValid())
 	{
 		KnownTextures.Add(PlatformGetTexturePlatform(NewTexture.Get()), NewTexture);
