Adding Rec.601 support to Electra video output

diff --git a/Engine/Plugins/Media/ElectraUtil/Source/ElectraSamples/Public/ElectraTextureSampleUtils.h b/Engine/Plugins/Media/ElectraUtil/Source/ElectraSamples/Public/ElectraTextureSampleUtils.h
--- a/Engine/Plugins/Media/ElectraUtil/Source/ElectraSamples/Public/ElectraTextureSampleUtils.h
+++ b/Engine/Plugins/Media/ElectraUtil/Source/ElectraSamples/Public/ElectraTextureSampleUtils.h
@@ -19,6 +19,9 @@ namespace ElectraColorimetryUtils
 			case 1:				// Rec709
 			case 2:				// unknown
 				return UE::Color::EColorSpace::sRGB;
+			case 6:				// Rec601
+			case 7:				// Rec601 (historical)
+				return UE::Color::EColorSpace::sRGB;
 			case 9:				// Rec2020
 				return UE::Color::EColorSpace::Rec2020;
 			default:
@@ -36,6 +39,9 @@ namespace ElectraColorimetryUtils
 			case 1:				// Rec709
 			case 2:				// unknown
 				return UE::Color::EColorSpace::sRGB;
+			case 6:				// Rec601
+			case 7:				// Rec601 (historical)
+				return UE::Color::EColorSpace::sRGB;
 			case 9:				// Rec2020
 				return UE::Color::EColorSpace::Rec2020;
 			default:
