Removing unnecessary D3D12 includes.

diff --git a/Engine/Plugins/Media/ElectraUtil/Source/ElectraSamples/Private/ElectraTextureSamplePC.cpp b/Engine/Plugins/Media/ElectraUtil/Source/ElectraSamples/Private/ElectraTextureSamplePC.cpp
--- a/Engine/Plugins/Media/ElectraUtil/Source/ElectraSamples/Private/ElectraTextureSamplePC.cpp
+++ b/Engine/Plugins/Media/ElectraUtil/Source/ElectraSamples/Private/ElectraTextureSamplePC.cpp
@@ -11,16 +11,10 @@
 
 #include "ID3D12DynamicRHI.h"
 
-#include "Windows/AllowWindowsPlatformTypes.h"
-THIRD_PARTY_INCLUDES_START
 #ifdef ELECTRA_HAVE_DX11
 #include "D3D11State.h"
 #include "D3D11Resources.h"
 #endif
-#include "d3d12.h"
-THIRD_PARTY_INCLUDES_END
-#include "Windows/HideWindowsPlatformTypes.h"
-
 
 /*
 	Short summary of how we get data:
