Fixing up state transitions for AVF Electra texture sample

diff --git a/Engine/Plugins/Media/ElectraUtil/Source/ElectraSamples/Private/ElectraTextureSampleAVF.cpp b/Engine/Plugins/Media/ElectraUtil/Source/ElectraSamples/Private/ElectraTextureSampleAVF.cpp
--- a/Engine/Plugins/Media/ElectraUtil/Source/ElectraSamples/Private/ElectraTextureSampleAVF.cpp
+++ b/Engine/Plugins/Media/ElectraUtil/Source/ElectraSamples/Private/ElectraTextureSampleAVF.cpp
@@ -165,8 +165,6 @@ void FElectraMediaTexConvApple::ConvertTexture(FTextureRHIRef & InDstTexture, CV
 	const int32 FrameWidth = CVPixelBufferGetWidth(InImageBufferRef);
 	const int32 FrameStride = CVPixelBufferGetBytesPerRow(InImageBufferRef);
 
-	TRefCountPtr<FRHITexture> ShaderResource;
-
 	// We have to support Metal for this object now
 	check(COREVIDEO_SUPPORTS_METAL);
 	check(IsMetalPlatform(GMaxRHIShaderPlatform));
@@ -211,12 +209,14 @@ void FElectraMediaTexConvApple::ConvertTexture(FTextureRHIRef & InDstTexture, CV
 			const FRHITextureCreateDesc YDesc =
 				FRHITextureCreateDesc::Create2D(TEXT("YTex"), YWidth, YHeight, bIs8Bit ? PF_G8 : PF_G16)
 				.SetFlags(ETextureCreateFlags::Dynamic | ETextureCreateFlags::NoTiling | ETextureCreateFlags::ShaderResource)
-				.SetBulkData(new FTexConvTexResourceWrapper(YTextureRef));
+				.SetBulkData(new FTexConvTexResourceWrapper(YTextureRef))
+				.SetInitialState(ERHIAccess::SRVMask);
 
 			const FRHITextureCreateDesc UVDesc =
 				FRHITextureCreateDesc::Create2D(TEXT("UVTex"), UVWidth, UVHeight, bIs8Bit ? PF_R8G8 : PF_G16R16)
 				.SetFlags(ETextureCreateFlags::Dynamic | ETextureCreateFlags::NoTiling | ETextureCreateFlags::ShaderResource)
-				.SetBulkData(new FTexConvTexResourceWrapper(UVTextureRef));
+				.SetBulkData(new FTexConvTexResourceWrapper(UVTextureRef))
+				.SetInitialState(ERHIAccess::SRVMask);
 
 			TRefCountPtr<FRHITexture> YTex = RHICreateTexture(YDesc);
 			TRefCountPtr<FRHITexture> UVTex = RHICreateTexture(UVDesc);
@@ -225,7 +225,8 @@ void FElectraMediaTexConvApple::ConvertTexture(FTextureRHIRef & InDstTexture, CV
 				// configure media shaders
 				auto GlobalShaderMap = GetGlobalShaderMap(GMaxRHIFeatureLevel);
 
-				FRHIRenderPassInfo RPInfo(InDstTexture, ERenderTargetActions::Load_Store);
+				RHICmdList.Transition(FRHITransitionInfo(InDstTexture, ERHIAccess::SRVMask, ERHIAccess::RTV));
+				FRHIRenderPassInfo RPInfo(InDstTexture, ERenderTargetActions::DontLoad_Store);
 				RHICmdList.BeginRenderPass(RPInfo, TEXT("AvfMediaSampler"));
 				{
 					FGraphicsPipelineStateInitializer GraphicsPSOInit;
@@ -286,7 +287,7 @@ void FElectraMediaTexConvApple::ConvertTexture(FTextureRHIRef & InDstTexture, CV
 					RHICmdList.DrawPrimitive(0, 2, 1);
 				}
 				RHICmdList.EndRenderPass();
-				RHICmdList.Transition(FRHITransitionInfo(ShaderResource, ERHIAccess::RTV, ERHIAccess::SRVMask));
+				RHICmdList.Transition(FRHITransitionInfo(InDstTexture, ERHIAccess::RTV, ERHIAccess::SRVMask));
 			}
 			CFRelease(YTextureRef);
 			CFRelease(UVTextureRef);
