ElectraHTTPStream: Cancelling a request in the generic implementation is not an error

#rnx

diff --git a/Engine/Plugins/Media/ElectraUtil/Source/ElectraHTTPStream/Private/Generic/GenericElectraHTTPStream.cpp b/Engine/Plugins/Media/ElectraUtil/Source/ElectraHTTPStream/Private/Generic/GenericElectraHTTPStream.cpp
--- a/Engine/Plugins/Media/ElectraUtil/Source/ElectraHTTPStream/Private/Generic/GenericElectraHTTPStream.cpp
+++ b/Engine/Plugins/Media/ElectraUtil/Source/ElectraHTTPStream/Private/Generic/GenericElectraHTTPStream.cpp
@@ -497,7 +497,7 @@ void FElectraHTTPStreamRequestGeneric::OnProcessRequestComplete(FHttpRequestPtr
 		Response->HTTPResponseCode = InHttpResponse->GetResponseCode();
 		Response->EffectiveURL = EffectiveURL;
 	}
-	if (!bInSucceeded)
+	if (!bInSucceeded && !WasCanceled())
 	{
 		if (Response->HTTPResponseCode)
 		{
