ElectraBase: added negative infinity to player time

#rnx

diff --git a/Engine/Plugins/Media/ElectraUtil/Source/ElectraBase/Public/PlayerTime.h b/Engine/Plugins/Media/ElectraUtil/Source/ElectraBase/Public/PlayerTime.h
--- a/Engine/Plugins/Media/ElectraUtil/Source/ElectraBase/Public/PlayerTime.h
+++ b/Engine/Plugins/Media/ElectraUtil/Source/ElectraBase/Public/PlayerTime.h
@@ -34,6 +34,12 @@ public:
 		return kPosInf;
 	}
 
+	static FTimeValue GetNegativeInfinity()
+	{
+		static FTimeValue kNegInf(-std::numeric_limits<double>::infinity());
+		return kNegInf;
+	}
+
 	static int64 MillisecondsToHNS(int64 InMilliseconds)
 	{
 		return InMilliseconds * 10000;
@@ -180,11 +186,20 @@ public:
 		return *this;
 	}
 
+	FTimeValue& SetToNegativeInfinity(int64 InSequenceIndex=0)
+	{
+		HNS = -0x7fffffffffffffffLL;
+		SequenceIndex = InSequenceIndex;
+		bIsValid = true;
+		bIsInfinity = true;
+		return *this;
+	}
+
 	FTimeValue& SetFromSeconds(double Seconds, int64 InSequenceIndex=0)
 	{
-		if ((bIsInfinity = (Seconds == std::numeric_limits<double>::infinity())) == true)
+		if ((bIsInfinity = (Seconds == std::numeric_limits<double>::infinity() || Seconds == -std::numeric_limits<double>::infinity())) == true)
 		{
-			HNS = 0x7fffffffffffffffLL;
+			HNS = Seconds < 0.0 ? -0x7fffffffffffffffLL : 0x7fffffffffffffffLL;
 			bIsValid = true;
 		}
 		else
@@ -280,7 +295,7 @@ public:
 	{
 		SequenceIndex = InSequenceIndex;
 	}
-	
+
 	int64 GetSequenceIndex() const
 	{
 		return SequenceIndex;
