ElectraHTTPStream: Changed error message on failed transfer

#rnx

diff --git a/Engine/Plugins/Media/ElectraUtil/Source/ElectraHTTPStream/Private/Generic/GenericElectraHTTPStream.cpp b/Engine/Plugins/Media/ElectraUtil/Source/ElectraHTTPStream/Private/Generic/GenericElectraHTTPStream.cpp
--- a/Engine/Plugins/Media/ElectraUtil/Source/ElectraHTTPStream/Private/Generic/GenericElectraHTTPStream.cpp
+++ b/Engine/Plugins/Media/ElectraUtil/Source/ElectraHTTPStream/Private/Generic/GenericElectraHTTPStream.cpp
@@ -501,7 +501,14 @@ void FElectraHTTPStreamRequestGeneric::OnProcessRequestComplete(FHttpRequestPtr
 	{
 		if (Response->HTTPResponseCode)
 		{
-			Response->SetErrorMessage(FString::Printf(TEXT("Failed with HTTP status %d"), Response->HTTPResponseCode));
+			if (Response->HTTPResponseCode >= 400)
+			{
+				Response->SetErrorMessage(FString::Printf(TEXT("Failed with HTTP status %d"), Response->HTTPResponseCode));
+			}
+			else
+			{
+				Response->SetErrorMessage(FString::Printf(TEXT("Failed due to connection error")));
+			}
 		}
 		else
 		{
