ElectraHTTPStream: Setup of individual response headers after a HEAD request

#rnx

diff --git a/Engine/Plugins/Media/ElectraUtil/Source/ElectraHTTPStream/Private/Generic/GenericElectraHTTPStream.cpp b/Engine/Plugins/Media/ElectraUtil/Source/ElectraHTTPStream/Private/Generic/GenericElectraHTTPStream.cpp
--- a/Engine/Plugins/Media/ElectraUtil/Source/ElectraHTTPStream/Private/Generic/GenericElectraHTTPStream.cpp
+++ b/Engine/Plugins/Media/ElectraUtil/Source/ElectraHTTPStream/Private/Generic/GenericElectraHTTPStream.cpp
@@ -491,6 +491,16 @@ void FElectraHTTPStreamRequestGeneric::SetFinished()
 
 void FElectraHTTPStreamRequestGeneric::OnProcessRequestComplete(FHttpRequestPtr InSourceHttpRequest, FHttpResponsePtr InHttpResponse, bool bInSucceeded)
 {
+	// If we did not receive any response body data bytes yet (because this is a HEAD request or due to some problem), process the headers.
+	if (CurrentState < EState::ReadingResponseData)
+	{
+		double Now = FPlatformTime::Seconds();
+		CurrentState = EState::ReadingResponseData;
+		Response->TimeUntilFirstByte = Now - Response->StartTime;
+		// Parse the headers and report them.
+		ParseResponseHeaders();
+	}
+
 	if (InHttpResponse.IsValid())
 	{
 		EffectiveURL = InHttpResponse->GetEffectiveURL();
@@ -567,14 +577,20 @@ void FElectraHTTPStreamRequestGeneric::OnHeaderReceived(FHttpRequestPtr InSource
 
 void FElectraHTTPStreamRequestGeneric::OnStatusCodeReceived(FHttpRequestPtr InSourceHttpRequest, int32 InHttpStatusCode)
 {
-	ReceivedHttpStatusCode = InHttpStatusCode;
+	if (InHttpStatusCode)
+	{
+		ReceivedHttpStatusCode = InHttpStatusCode;
+	}
 	if (WasCanceled())
 	{
 		return;
 	}
-	// For the lack of better knowledge pretend this is a 1.1 transfer.
-	FString HeaderValue = FString::Printf(TEXT("HTTP/1.1 %d"), InHttpStatusCode);
-	OnHeaderReceived(InSourceHttpRequest, FString(), HeaderValue);
+	if (InHttpStatusCode)
+	{
+		// For the lack of better knowledge pretend this is a 1.1 transfer.
+		FString HeaderValue = FString::Printf(TEXT("HTTP/1.1 %d"), InHttpStatusCode);
+		OnHeaderReceived(InSourceHttpRequest, FString(), HeaderValue);
+	}
 }
 
 bool FElectraHTTPStreamRequestGeneric::OnProcessRequestStream(void *InDataPtr, int64 InLength)
