Deprecated the v1 api to set http response stream delegate.

diff --git a/Engine/Plugins/Media/ElectraPlayer/Source/ElectraPlayerRuntime/Private/Runtime/Player/mpegaudio/PlaylistReaderMPEGAudio.cpp b/Engine/Plugins/Media/ElectraPlayer/Source/ElectraPlayerRuntime/Private/Runtime/Player/mpegaudio/PlaylistReaderMPEGAudio.cpp
--- a/Engine/Plugins/Media/ElectraPlayer/Source/ElectraPlayerRuntime/Private/Runtime/Player/mpegaudio/PlaylistReaderMPEGAudio.cpp
+++ b/Engine/Plugins/Media/ElectraPlayer/Source/ElectraPlayerRuntime/Private/Runtime/Player/mpegaudio/PlaylistReaderMPEGAudio.cpp
@@ -109,7 +109,7 @@ private:
 		void OnProcessRequestComplete(FHttpRequestPtr InSourceHttpRequest, FHttpResponsePtr InHttpResponse, bool bInSucceeded);
 		void OnHeaderReceived(FHttpRequestPtr InSourceHttpRequest, const FString& InHeaderName, const FString& InHeaderValue);
 		void OnStatusCodeReceived(FHttpRequestPtr InSourceHttpRequest, int32 InHttpStatusCode);
-		bool OnProcessRequestStream(void* InDataPtr, int64 InLength);
+		void OnProcessRequestStream(void* InDataPtr, int64& InLength);
 		void FindSyncMarkers();
 		void Cancel();
 		bool Validate(FMPEGAudioInfoHeader& InOutInfoHeader, FString& OutError, TArray<FString>& OutWarnings);
@@ -318,11 +318,11 @@ void FPlaylistReaderMPEGAudio::FRequest::FindSyncMarkers()
 	}
 }
 
-bool FPlaylistReaderMPEGAudio::FRequest::OnProcessRequestStream(void *InDataPtr, int64 InLength)
+void FPlaylistReaderMPEGAudio::FRequest::OnProcessRequestStream(void *InDataPtr, int64& InLength)
 {
 	if (StatusCode < 200 || StatusCode >= 300)
 	{
-		return true;
+		return;
 	}
 
 	// Add new data unconditionally. This won't be overly much so not to worry.
@@ -356,7 +356,6 @@ bool FPlaylistReaderMPEGAudio::FRequest::OnProcessRequestStream(void *InDataPtr,
 		bSentHaveProbeDataMsg = true;
 		Events.SendMessage(EEvent::HaveProbeData);
 	}
-	return true;
 }
 
 void FPlaylistReaderMPEGAudio::FRequest::Cancel()
@@ -833,7 +832,8 @@ void FPlaylistReaderMPEGAudio::WorkerThread()
 				AlreadyRead = 10;
 				Req->Buffer.AddUninitialized(10);
 				Archive->Serialize(Req->Buffer.GetData(), 10);
-				Req->OnProcessRequestStream(nullptr, 0);
+				int64 Length = 0;
+				Req->OnProcessRequestStream(nullptr, Length);
 				check(Req->ProbeSize > 0);
 			}
 			Req->ProbeSize = Req->ProbeSize < 0 ? FRequest::GetMinProbeSize() : Req->ProbeSize;
@@ -850,7 +850,7 @@ void FPlaylistReaderMPEGAudio::WorkerThread()
 	}
 	else
 	{
-		FHttpRequestStreamDelegate StreamDelegate;
+		FHttpRequestStreamDelegateV2 StreamDelegate;
 		StreamDelegate.BindThreadSafeSP(Req.ToSharedRef(), &FPlaylistReaderMPEGAudio::FRequest::OnProcessRequestStream);
 		Req->Handle = FHttpModule::Get().CreateRequest();
 		Req->Handle->SetVerb(TEXT("GET"));
@@ -859,7 +859,7 @@ void FPlaylistReaderMPEGAudio::WorkerThread()
 		Req->Handle->OnProcessRequestComplete().BindThreadSafeSP(Req.ToSharedRef(), &FPlaylistReaderMPEGAudio::FRequest::OnProcessRequestComplete);
 		Req->Handle->OnHeaderReceived().BindThreadSafeSP(Req.ToSharedRef(), &FPlaylistReaderMPEGAudio::FRequest::OnHeaderReceived);
 		Req->Handle->OnStatusCodeReceived().BindThreadSafeSP(Req.ToSharedRef(), &FPlaylistReaderMPEGAudio::FRequest::OnStatusCodeReceived);
-		Req->Handle->SetResponseBodyReceiveStreamDelegate(StreamDelegate);
+		Req->Handle->SetResponseBodyReceiveStreamDelegateV2(StreamDelegate);
 		Req->Handle->SetHeader(TEXT("User-Agent"), IElectraHttpManager::GetDefaultUserAgent());
 		Req->Handle->SetHeader(TEXT("Accept-Encoding"), TEXT("identity"));
 		Req->Handle->SetTimeout(PlayerSessionServices->GetOptionValue(MPEGAudio::OptionKeyMPEGAudioLoadTimeout).SafeGetTimeValue(FTimeValue().SetFromMilliseconds(1000 * 60)).GetAsSeconds());
diff --git a/Engine/Plugins/Media/ElectraPlayer/Source/ElectraPlayerRuntime/Private/Runtime/Player/mpegaudio/StreamReaderMPEGAudio.cpp b/Engine/Plugins/Media/ElectraPlayer/Source/ElectraPlayerRuntime/Private/Runtime/Player/mpegaudio/StreamReaderMPEGAudio.cpp
--- a/Engine/Plugins/Media/ElectraPlayer/Source/ElectraPlayerRuntime/Private/Runtime/Player/mpegaudio/StreamReaderMPEGAudio.cpp
+++ b/Engine/Plugins/Media/ElectraPlayer/Source/ElectraPlayerRuntime/Private/Runtime/Player/mpegaudio/StreamReaderMPEGAudio.cpp
@@ -61,7 +61,7 @@ struct FStreamReaderMPEGAudio::FLiveRequest : public TSharedFromThis<FStreamRead
 	void OnProcessRequestComplete(FHttpRequestPtr InSourceHttpRequest, FHttpResponsePtr InHttpResponse, bool bInSucceeded);
 	void OnHeaderReceived(FHttpRequestPtr InSourceHttpRequest, const FString& InHeaderName, const FString& InHeaderValue);
 	void OnStatusCodeReceived(FHttpRequestPtr InSourceHttpRequest, int32 InHttpStatusCode);
-	bool OnProcessRequestStream(void* InDataPtr, int64 InLength);
+	void OnProcessRequestStream(void* InDataPtr, int64& InLength);
 	void Cancel();
 	void WaitUntilFinished();
 	enum class EEvent
@@ -71,7 +71,7 @@ struct FStreamReaderMPEGAudio::FLiveRequest : public TSharedFromThis<FStreamRead
 		Finished
 	};
 	TSharedPtr<IHttpRequest, ESPMode::ThreadSafe> Handle;
-	FHttpRequestStreamDelegate StreamDelegate;
+	FHttpRequestStreamDelegateV2 StreamDelegate;
 	int32 StatusCode = 0;
 	TMediaMessageQueueWithTimeout<EEvent> Events;
 	TMultiMap<FString, FString> Headers;
@@ -253,7 +253,7 @@ void FStreamReaderMPEGAudio::FLiveRequest::OnStatusCodeReceived(FHttpRequestPtr
 	}
 }
 
-bool FStreamReaderMPEGAudio::FLiveRequest::OnProcessRequestStream(void* InDataPtr, int64 InLength)
+void FStreamReaderMPEGAudio::FLiveRequest::OnProcessRequestStream(void* InDataPtr, int64& InLength)
 {
 	if (StatusCode == 200)
 	{
@@ -295,11 +295,13 @@ bool FStreamReaderMPEGAudio::FLiveRequest::OnProcessRequestStream(void* InDataPt
 					int64 BufSizeRequired = rb->Buffer.Num() + DataBytesNow;
 					if (!rb->Buffer.EnlargeTo(BufSizeRequired))
 					{
-						return false;
+						InLength = 0;
+						return;
 					}
 					if (!rb->Buffer.PushData(reinterpret_cast<const uint8*>(InDataPtr), DataBytesNow))
 					{
-						return false;
+						InLength = 0;
+						return;
 					}
 
 					// Sanity check that we are not reading excessive data, which is the case when the
@@ -307,7 +309,8 @@ bool FStreamReaderMPEGAudio::FLiveRequest::OnProcessRequestStream(void* InDataPt
 					if (MaxDataBytes && rb->Buffer.Num() > MaxDataBytes)
 					{
 						bHasFailed = true;
-						return false;
+						InLength = 0;
+						return;
 					}
 				}
 				TotalDataBytePos += DataBytesNow;
@@ -322,7 +325,6 @@ bool FStreamReaderMPEGAudio::FLiveRequest::OnProcessRequestStream(void* InDataPt
 			}
 		}
 	}
-	return true;
 }
 
 
@@ -416,7 +418,7 @@ void FStreamReaderMPEGAudio::HandleRequest()
 		LiveRequest->Handle->OnProcessRequestComplete().BindThreadSafeSP(LiveRequest.ToSharedRef(), &FStreamReaderMPEGAudio::FLiveRequest::OnProcessRequestComplete);
 		LiveRequest->Handle->OnHeaderReceived().BindThreadSafeSP(LiveRequest.ToSharedRef(), &FStreamReaderMPEGAudio::FLiveRequest::OnHeaderReceived);
 		LiveRequest->Handle->OnStatusCodeReceived().BindThreadSafeSP(LiveRequest.ToSharedRef(), &FStreamReaderMPEGAudio::FLiveRequest::OnStatusCodeReceived);
-		LiveRequest->Handle->SetResponseBodyReceiveStreamDelegate(LiveRequest->StreamDelegate);
+		LiveRequest->Handle->SetResponseBodyReceiveStreamDelegateV2(LiveRequest->StreamDelegate);
 		LiveRequest->Handle->SetHeader(TEXT("User-Agent"), IElectraHttpManager::GetDefaultUserAgent());
 		LiveRequest->Handle->SetHeader(TEXT("Accept-Encoding"), TEXT("identity"));
 		// If this is an Icycast, we ask for period metadata.
diff --git a/Engine/Plugins/Media/ElectraUtil/Source/ElectraHTTPStream/Private/Generic/GenericElectraHTTPStream.cpp b/Engine/Plugins/Media/ElectraUtil/Source/ElectraHTTPStream/Private/Generic/GenericElectraHTTPStream.cpp
--- a/Engine/Plugins/Media/ElectraUtil/Source/ElectraHTTPStream/Private/Generic/GenericElectraHTTPStream.cpp
+++ b/Engine/Plugins/Media/ElectraUtil/Source/ElectraHTTPStream/Private/Generic/GenericElectraHTTPStream.cpp
@@ -35,15 +35,6 @@ DECLARE_STATS_GROUP(TEXT("Electra HTTP Stream"), STATGROUP_ElectraHTTPStream, ST
 DECLARE_CYCLE_STAT(TEXT("Process"), STAT_ElectraHTTPThread_Process, STATGROUP_ElectraHTTPStream);
 DECLARE_CYCLE_STAT(TEXT("Custom handler"), STAT_ElectraHTTPThread_CustomHandler, STATGROUP_ElectraHTTPStream);
 
-/***************************************************************************************************************************************************/
-/***************************************************************************************************************************************************/
-/***************************************************************************************************************************************************/
-TAutoConsoleVariable<bool> CVarElectraUseHttpStreamDelegateV2(
-	TEXT("electra.UseHttpStreamDelegateV2"),
-	true,
-	TEXT("Use http stream delegate v2 in electra")
-);
-
 /***************************************************************************************************************************************************/
 /***************************************************************************************************************************************************/
 /***************************************************************************************************************************************************/
@@ -253,8 +244,7 @@ private:
 	void OnProcessRequestComplete(FHttpRequestPtr InSourceHttpRequest, FHttpResponsePtr InHttpResponse, bool bInSucceeded);
 	void OnHeaderReceived(FHttpRequestPtr InSourceHttpRequest, const FString& InHeaderName, const FString& InHeaderValue);
 	void OnStatusCodeReceived(FHttpRequestPtr InSourceHttpRequest, int32 InHttpStatusCode);
-	bool OnProcessRequestStream(void *InDataPtr, int64 InLength);
-	void OnProcessRequestStreamV2(void *InDataPtr, int64& InLength);
+	void OnProcessRequestStream(void *InDataPtr, int64& InLength);
 
 	// Owner to be notified on activity.
 	TWeakPtr<FElectraHTTPStreamGeneric, ESPMode::ThreadSafe> Owner;
@@ -336,20 +326,10 @@ bool FElectraHTTPStreamRequestGeneric::Setup()
 	RequestHandle->OnHeaderReceived().BindThreadSafeSP(AsShared(), &FElectraHTTPStreamRequestGeneric::OnHeaderReceived);
 	RequestHandle->OnStatusCodeReceived().BindThreadSafeSP(AsShared(), &FElectraHTTPStreamRequestGeneric::OnStatusCodeReceived);
 
-	if (CVarElectraUseHttpStreamDelegateV2.GetValueOnAnyThread())
-	{
-		FHttpRequestStreamDelegateV2 StreamDelegate;
-		StreamDelegate.BindThreadSafeSP(AsShared(), &FElectraHTTPStreamRequestGeneric::OnProcessRequestStreamV2);
-		bool bOk = RequestHandle->SetResponseBodyReceiveStreamDelegateV2(StreamDelegate);
-		(void)bOk; check(bOk);
-	}
-	else
-	{
-		FHttpRequestStreamDelegate StreamDelegate;
-		StreamDelegate.BindThreadSafeSP(AsShared(), &FElectraHTTPStreamRequestGeneric::OnProcessRequestStream);
-		bool bOk = RequestHandle->SetResponseBodyReceiveStreamDelegate(StreamDelegate);
-		(void)bOk; check(bOk);
-	}
+	FHttpRequestStreamDelegateV2 StreamDelegate;
+	StreamDelegate.BindThreadSafeSP(AsShared(), &FElectraHTTPStreamRequestGeneric::OnProcessRequestStream);
+	bool bOk = RequestHandle->SetResponseBodyReceiveStreamDelegateV2(StreamDelegate);
+	(void)bOk; check(bOk);
 
 	// We set the user agent manually. For simplicities sake we add it to the list of additional headers.
 	AdditionalHeaders.FindOrAdd(TEXT("User-Agent"), UserAgent);
@@ -614,45 +594,7 @@ void FElectraHTTPStreamRequestGeneric::OnStatusCodeReceived(FHttpRequestPtr InSo
 	}
 }
 
-bool FElectraHTTPStreamRequestGeneric::OnProcessRequestStream(void *InDataPtr, int64 InLength)
-{
-	if (InDataPtr == nullptr || InLength < 0)
-	{
-		return false;
-	}
-	if (WasCanceled())
-	{
-		return true;
-	}
-	double Now = FPlatformTime::Seconds();
-	if (CurrentState < EState::ReadingResponseData)
-	{
-		CurrentState = EState::ReadingResponseData;
-		Response->TimeUntilFirstByte = Now - Response->StartTime;
-		// Parse the headers and report them.
-		if (!ParseResponseHeaders())
-		{
-			return false;
-		}
-	}
-
-	// Add the data to the response.
-	TConstArrayView<const uint8> Data(static_cast<const uint8*>(InDataPtr), (int32)InLength);
-	Response->AddResponseData(Data);
-	// Notify amount of new data available.
-	NotifyCallback(EElectraHTTPStreamNotificationReason::ReadData, InLength);
-
-	Response->TimeOfMostRecentReceive = Now;
-
-	TSharedPtr<FElectraHTTPStreamGeneric, ESPMode::ThreadSafe> PinnedOwner = Owner.Pin();
-	if (PinnedOwner.IsValid())
-	{
-		PinnedOwner->TriggerWorkSignal();
-	}
-	return true;
-}
-
-void FElectraHTTPStreamRequestGeneric::OnProcessRequestStreamV2(void *InDataPtr, int64& InOutLength)
+void FElectraHTTPStreamRequestGeneric::OnProcessRequestStream(void *InDataPtr, int64& InOutLength)
 {
 	if (InDataPtr == nullptr || InOutLength < 0)
 	{
diff --git a/Engine/Source/Runtime/Online/HTTP/Public/Interfaces/IHttpRequest.h b/Engine/Source/Runtime/Online/HTTP/Public/Interfaces/IHttpRequest.h
--- a/Engine/Source/Runtime/Online/HTTP/Public/Interfaces/IHttpRequest.h
+++ b/Engine/Source/Runtime/Online/HTTP/Public/Interfaces/IHttpRequest.h
@@ -239,6 +239,7 @@ public:
 	 * @param StreamDelegate - will be used to receive the response body
 	 * @return True if the delegate can be used. False otherwise.
 	 */
+	UE_DEPRECATED(5.5, "SetResponseBodyReceiveStreamDelegate has been deprecated, use SetResponseBodyReceiveStreamDelegateV2 instead")
 	HTTP_API bool SetResponseBodyReceiveStreamDelegate(FHttpRequestStreamDelegate StreamDelegate);
 
 	/**
