[Media] Fixing NotchLC Electra Decoder (d3d12 path)

- Fixing the d3d11 to d3d12 texture conversion path (when RHI is d3d12). The d3d12 path had been added in CL 30379413. The d3d11 texture conversion path is only used by NotchLC (afaik) and probably got missed.
- Fixing NotchLC decoder's d3d12 setup code. The CreateDXGIFactory1 must be used to have the ability to create shared keyed mutex d3d11 textures. Otherwise, CreateTexture2D fails in FElectraPlayerVideoDecoderOutputPC::InitializeWithSharedTexture.

diff --git a/Engine/Plugins/Media/ElectraUtil/Source/ElectraSamples/Private/ElectraTextureSamplePC.cpp b/Engine/Plugins/Media/ElectraUtil/Source/ElectraSamples/Private/ElectraTextureSamplePC.cpp
--- a/Engine/Plugins/Media/ElectraUtil/Source/ElectraSamples/Private/ElectraTextureSamplePC.cpp
+++ b/Engine/Plugins/Media/ElectraUtil/Source/ElectraSamples/Private/ElectraTextureSamplePC.cpp
@@ -202,6 +202,23 @@ FRHITexture* FElectraTextureSample::GetTexture() const
 				// Yes...
 				TRefCountPtr<ID3D12Resource> TextureDX12;
 				HRESULT Res = TextureCommon->QueryInterface(__uuidof(ID3D12Resource), (void**)TextureDX12.GetInitReference());
+
+				if (Res != S_OK)
+				{
+					// Support shared dxgi resource.
+					TRefCountPtr<IDXGIResource> DxgiResource;
+					Res = TextureCommon->QueryInterface(__uuidof(IDXGIResource), (void**)DxgiResource.GetInitReference());
+					if (Res == S_OK)
+					{
+						HANDLE SharedHandle;
+						Res = DxgiResource->GetSharedHandle(&SharedHandle);
+						if (Res == S_OK)
+						{
+							Res = GetID3D12DynamicRHI()->RHIGetDevice(0)->OpenSharedHandle(SharedHandle, __uuidof(ID3D12Resource), (void**)TextureDX12.GetInitReference());
+						}
+					}
+				}
+
 				check(Res == S_OK);
 				if (Res == S_OK)
 				{
@@ -366,14 +383,40 @@ bool FElectraTextureSample::Convert(FRHICommandListImmediate& RHICmdList, FTextu
 			//
 			TRefCountPtr<ID3D12Fence> SyncFence;
 			HRESULT Res = SyncCommon->QueryInterface(__uuidof(ID3D12Fence), (void**)SyncFence.GetInitReference());
-			check(SUCCEEDED(Res));
+			
 			if (Res == S_OK)
 			{
 				RHICmdList.EnqueueLambda([SyncFence, SyncFenceValue](FRHICommandList& RHICmdList)
 				{
 					GetID3D12DynamicRHI()->RHIWaitManualFence(RHICmdList, SyncFence, SyncFenceValue);
 				});
+				return true;
+			}
+
+			// Support IDXGIKeyedMutex (d3d11 texture path)
+			TRefCountPtr<IDXGIKeyedMutex> KeyedMutex;
+			Res = SyncCommon->QueryInterface(_uuidof(IDXGIKeyedMutex), (void**)&KeyedMutex);
+			if (Res == S_OK)
+			{
+				// Remark:
+				// In the d3d11 texture path, the VideoDecoderOutputPC::InitializeWithSharedTexture should have already flushed d3d11 context and
+				// the AcquireSync will not wait in that case. Adding the sync command in any case for extra protection.
+
+				RHICmdList.EnqueueLambda([KeyedMutex](FRHICommandList& RHICmdList)
+				{
+					if (KeyedMutex)
+					{
+						// Should we limit the wait as a precaution? ex: 16 ms instead of infinite?
+						if (KeyedMutex->AcquireSync(1, INFINITE) == S_OK)
+						{
+							KeyedMutex->ReleaseSync(2);
+						}
+					}
+				});
+				return true;
 			}
+
+			check(SUCCEEDED(Res));
 		}
 		return true;
 	}
