[Electra Player] d3d11 texture sample conversion path: fallback to keyed mutex path only if d3d12 interface was not available so that we don't loose the other errors.

Addressing review comment from CL 35370165.

diff --git a/Engine/Plugins/Media/ElectraUtil/Source/ElectraSamples/Private/ElectraTextureSamplePC.cpp b/Engine/Plugins/Media/ElectraUtil/Source/ElectraSamples/Private/ElectraTextureSamplePC.cpp
--- a/Engine/Plugins/Media/ElectraUtil/Source/ElectraSamples/Private/ElectraTextureSamplePC.cpp
+++ b/Engine/Plugins/Media/ElectraUtil/Source/ElectraSamples/Private/ElectraTextureSamplePC.cpp
@@ -393,27 +393,31 @@ bool FElectraTextureSample::Convert(FRHICommandListImmediate& RHICmdList, FTextu
 				return true;
 			}
 
-			// Support IDXGIKeyedMutex (d3d11 texture path)
-			TRefCountPtr<IDXGIKeyedMutex> KeyedMutex;
-			Res = SyncCommon->QueryInterface(_uuidof(IDXGIKeyedMutex), (void**)&KeyedMutex);
-			if (Res == S_OK)
+			// Only continue with fallback if no interface so that we don't loose the other errors.
+			if (Res == E_NOINTERFACE)
 			{
-				// Remark:
-				// In the d3d11 texture path, the VideoDecoderOutputPC::InitializeWithSharedTexture should have already flushed d3d11 context and
-				// the AcquireSync will not wait in that case. Adding the sync command in any case for extra protection.
-
-				RHICmdList.EnqueueLambda([KeyedMutex](FRHICommandList& RHICmdList)
+				// Support IDXGIKeyedMutex (d3d11 texture path)
+				TRefCountPtr<IDXGIKeyedMutex> KeyedMutex;
+				Res = SyncCommon->QueryInterface(_uuidof(IDXGIKeyedMutex), (void**)&KeyedMutex);
+				if (Res == S_OK)
 				{
-					if (KeyedMutex)
+					// Remark:
+					// In the d3d11 texture path, the VideoDecoderOutputPC::InitializeWithSharedTexture should have already flushed d3d11 context and
+					// the AcquireSync will not wait in that case. Adding the sync command in any case for extra protection.
+
+					RHICmdList.EnqueueLambda([KeyedMutex](FRHICommandList& RHICmdList)
 					{
-						// Should we limit the wait as a precaution? ex: 16 ms instead of infinite?
-						if (KeyedMutex->AcquireSync(1, INFINITE) == S_OK)
+						if (KeyedMutex)
 						{
-							KeyedMutex->ReleaseSync(2);
+							// Should we limit the wait as a precaution? ex: 16 ms instead of infinite?
+							if (KeyedMutex->AcquireSync(1, INFINITE) == S_OK)
+							{
+								KeyedMutex->ReleaseSync(2);
+							}
 						}
-					}
-				});
-				return true;
+					});
+					return true;
+				}
 			}
 
 			check(SUCCEEDED(Res));
