ElectraPlayer Android improvements
Added experimental GPU output path to Vulkan renderer that can be enabled by setting Electra.AndroidUseGpuOutputPath to 1

#rnx

diff --git a/Engine/Plugins/Media/ElectraUtil/Source/ElectraSamples/ElectraSamples.Build.cs b/Engine/Plugins/Media/ElectraUtil/Source/ElectraSamples/ElectraSamples.Build.cs
--- a/Engine/Plugins/Media/ElectraUtil/Source/ElectraSamples/ElectraSamples.Build.cs
+++ b/Engine/Plugins/Media/ElectraUtil/Source/ElectraSamples/ElectraSamples.Build.cs
@@ -74,6 +74,8 @@ namespace UnrealBuildTool.Rules
 				PublicIncludePaths.Add("$(ModuleDir)/Public/Android");
 				string PluginPath = Utils.MakePathRelativeTo(ModuleDirectory, Target.RelativeEnginePath);
 				AdditionalPropertiesForReceipt.Add("AndroidPlugin", Path.Combine(PluginPath, "ElectraSamples_UPL.xml"));
+				PrivateDependencyModuleNames.Add("VulkanRHI");
+				AddEngineThirdPartyPrivateStaticDependencies(Target, "Vulkan");
 			}
 			else if (Target.IsInPlatformGroup(UnrealPlatformGroup.Unix))
 			{
diff --git a/Engine/Plugins/Media/ElectraUtil/Source/ElectraSamples/Private/ElectraTextureSampleAndroid.cpp b/Engine/Plugins/Media/ElectraUtil/Source/ElectraSamples/Private/ElectraTextureSampleAndroid.cpp
--- a/Engine/Plugins/Media/ElectraUtil/Source/ElectraSamples/Private/ElectraTextureSampleAndroid.cpp
+++ b/Engine/Plugins/Media/ElectraUtil/Source/ElectraSamples/Private/ElectraTextureSampleAndroid.cpp
@@ -5,6 +5,7 @@
 #if PLATFORM_ANDROID
 
 #include "ElectraTextureSample.h"
+#include "ElectraSamplesModule.h"
 
 #include "Android/AndroidPlatform.h"
 #include "Android/AndroidJava.h"
@@ -14,8 +15,34 @@
 #include "RHIStaticStates.h"
 #include "PipelineStateCache.h"
 
+#include "RenderUtils.h"
+
+#include <EGL/egl.h>
+#include <EGL/eglext.h>
+#include <GLES3/gl31.h>
+#include "GLES2/gl2ext.h"
+#include "android/hardware_buffer.h"
+#include "android/hardware_buffer_jni.h"
+
+#include "VulkanCommon.h"
+#include "IVulkanDynamicRHI.h"
+
+#include "vulkan/vulkan_android.h"
+
+
 DECLARE_GPU_STAT_NAMED(MediaAndroidDecoder_Convert, TEXT("MediaAndroidDecoder_Convert"));
 
+/*********************************************************************************************************************/
+
+static TAutoConsoleVariable<int32> CVarElectraAndroidUseGpuOutputPath(
+	TEXT("Electra.AndroidUseGpuOutputPath"),
+	0,
+	TEXT("Use experimental direct to GPU output path on Android.\n")
+	TEXT(" 0: use CPU output path (default); 1: use new direct to GPU output path."),
+	ECVF_Default);
+
+/*********************************************************************************************************************/
+
 #define ELECTRA_INIT_ON_RENDERTHREAD	1	// set to 1 if context & surface init should be on render thread (seems safer for compatibility - TODO: research why)
 
 // ---------------------------------------------------------------------------------------------------------------------
@@ -41,6 +68,64 @@ DECLARE_GPU_STAT_NAMED(MediaAndroidDecoder_Convert, TEXT("MediaAndroidDecoder_Co
 
 // ---------------------------------------------------------------------------------------------------------------------
 
+namespace {
+
+struct FElectraTextureSampleVulkanResources
+{
+	VkImage Image;
+	VkDeviceMemory DeviceMemory;
+};
+
+static void CleanupImageResourcesVulkan(void* UserData)
+{
+	check(UserData);
+
+	FElectraTextureSampleVulkanResources* VulkanResources = static_cast<FElectraTextureSampleVulkanResources*>(UserData);
+
+	IVulkanDynamicRHI* RHI = GetIVulkanDynamicRHI();
+	VkDevice Device = RHI->RHIGetVkDevice();
+	const VkAllocationCallbacks* AllocationCallbacks = RHI->RHIGetVkAllocationCallbacks();
+
+	if (VulkanResources->DeviceMemory != VK_NULL_HANDLE)
+	{
+		static PFN_vkFreeMemory VkFreeMemory = nullptr;
+		if (!VkFreeMemory)
+		{
+			VkFreeMemory = (PFN_vkFreeMemory)RHI->RHIGetVkInstanceProcAddr("vkFreeMemory");
+			check(VkFreeMemory);
+		}
+
+		VkFreeMemory(Device, VulkanResources->DeviceMemory, AllocationCallbacks);
+	}
+
+	if (VulkanResources->Image != VK_NULL_HANDLE)
+	{
+		static PFN_vkDestroyImage VkDestroyImage = nullptr;
+		if (!VkDestroyImage)
+		{
+			VkDestroyImage = (PFN_vkDestroyImage)RHI->RHIGetVkInstanceProcAddr("vkDestroyImage");
+			check(VkDestroyImage);
+		}
+
+		VkDestroyImage(Device, VulkanResources->Image, AllocationCallbacks);
+	}
+
+	delete VulkanResources;
+}
+
+static void CleanupImageResourcesJNI(jobject Resources, jmethodID ReleaseFN)
+{
+	QUICK_SCOPE_CYCLE_COUNTER(STAT_FRHICommandReleaseDecoderResources_Execute);
+
+	JNIEnv* JEnv = FAndroidApplication::GetJavaEnv();
+	JEnv->CallVoidMethod(Resources, ReleaseFN);
+	JEnv->DeleteGlobalRef(Resources);
+}
+
+} // namespace
+
+// ---------------------------------------------------------------------------------------------------------------------
+
 class FElectraTextureSampleSupport : public FJavaClassObject
 {
 public:
@@ -49,21 +134,57 @@ public:
 
 	struct FFrameUpdateInfo
 	{
+		~FFrameUpdateInfo()
+		{
+			if (ImageResources)
+			{
+				JNIEnv* JEnv = FAndroidApplication::GetJavaEnv();
+				//TODO: CALL "Release" - OR REMOVE THIS VERSION IF WE NEVER USE IT!
+				JEnv->DeleteGlobalRef(ImageResources);
+			}
+		}
+
+		jobject MoveImageResources()
+		{
+			jobject Ret = ImageResources;
+			ImageResources = nullptr;
+			return Ret;
+		}
+
 		int64   Timestamp = 0;
 		int64   Duration = 0;
 		float   UScale = 0.0f;
 		float   UOffset = 0.0f;
 		float   VScale = 0.0f;
 		float   VOffset = 0.0f;
-		bool	bFrameReady = false;
-		bool	bRegionChanged = false;
-		int 	NumPending = 0;
+		bool    bFrameReady = false;
+		bool    bRegionChanged = false;
+		int     NumPending = 0;
+		jobject ImageResources = nullptr;
 	};
 
 	int32 GetFrameDataAndUpdateInfo(FFrameUpdateInfo& OutFrameUpdateInfo, FElectraTextureSample* InTargetSample);
 	int32 GetFrameData(FElectraTextureSample* InTargetSample);
 	jobject GetCodecSurface();
 
+	jmethodID GetImageResources_ReleaseFN() const
+	{
+		return FImageResources_ReleaseFN;
+	}
+
+	void SignalImageReaderSurfaceRead()
+	{
+		if (CodecSurfaceReadEvent)
+		{
+			CodecSurfaceReadEvent->Trigger();
+		}
+	}
+
+	jobject ImageResources_GetHardwareBuffer(jobject ImageResources);
+	void ImageResources_GetScaleOffset(jobject ImageResources, FVector2f& OutScale, FVector2f& OutOffset);
+
+	const bool UseGpuOutputPath;
+
 private:
 	// Java methods
 	FJavaClassMethod	InitializeFN;
@@ -83,12 +204,22 @@ private:
 	jfieldID			FFrameUpdateInfo_VScale;
 	jfieldID			FFrameUpdateInfo_VOffset;
 	jfieldID			FFrameUpdateInfo_NumPending;
+	jfieldID			FFrameUpdateInfo_ImageResources;
+
+	// FImageResources members / methods
+	jclass				FImageResourcesClass;
+	jfieldID			FImageResources_HardwareBufferHandle;
+	jfieldID			FImageResources_UScale;
+	jfieldID			FImageResources_VScale;
+	jfieldID			FImageResources_UOffset;
+	jfieldID			FImageResources_VOffset;
+	jmethodID			FImageResources_ReleaseFN;
 
 	jobject				CodecSurface;
 	FEvent*				SurfaceInitEvent;
 	FCriticalSection	CodecSurfaceLock;
 	jobject				CodecSurfaceToDelete;
-
+	FEvent*				CodecSurfaceReadEvent;
 
 	static FName GetClassName()
 	{
@@ -124,13 +255,15 @@ private:
 
 FElectraTextureSampleSupport::FElectraTextureSampleSupport()
 	: FJavaClassObject(GetClassName(), "()V")
-	, InitializeFN(GetClassMethod("Initialize", "(Z)V"))
+	, UseGpuOutputPath(FAndroidMisc::ShouldUseVulkan() && (CVarElectraAndroidUseGpuOutputPath.GetValueOnAnyThread() != 0)) // GpuOutputPath is only available for Vulkan right now (and experimental)
+	, InitializeFN(GetClassMethod("Initialize", "(ZZJ)V"))
 	, ReleaseFN(GetClassMethod("Release", "()V"))
 	, GetCodecSurfaceFN(GetClassMethod("GetCodecSurface", "()Landroid/view/Surface;"))
 	, GetVideoFrameUpdateInfoFN(GetClassMethod("GetVideoFrameUpdateInfo", "(IIIZ)Lcom/epicgames/unreal/ElectraTextureSample$FFrameUpdateInfo;"))
 	, CodecSurface(nullptr)
 	, SurfaceInitEvent(nullptr)
 	, CodecSurfaceToDelete(nullptr)
+	, CodecSurfaceReadEvent(nullptr)
 {
 	JNIEnv* JEnv = FAndroidApplication::GetJavaEnv();
 
@@ -148,7 +281,50 @@ FElectraTextureSampleSupport::FElectraTextureSampleSupport()
 	FFrameUpdateInfo_VScale = FindField(JEnv, FFrameUpdateInfoClass, "VScale", "F", false);
 	FFrameUpdateInfo_VOffset = FindField(JEnv, FFrameUpdateInfoClass, "VOffset", "F", false);
 	FFrameUpdateInfo_NumPending = FindField(JEnv, FFrameUpdateInfoClass, "NumPending", "I", false);
+	FFrameUpdateInfo_ImageResources = FindField(JEnv, FFrameUpdateInfoClass, "ImageResources", "Lcom/epicgames/unreal/ElectraTextureSample$FImageResources;", false);
+
+	// Get field IDs for FImageResources class members (etc.)
+	jclass localImageResourcesClass = FAndroidApplication::FindJavaClass("com/epicgames/unreal/ElectraTextureSample$FImageResources");
+	FImageResourcesClass = (jclass)JEnv->NewGlobalRef(localImageResourcesClass);
+	JEnv->DeleteLocalRef(localImageResourcesClass);
+	FImageResources_HardwareBufferHandle = FindField(JEnv, FImageResourcesClass, "HardwareBuffer", "Landroid/hardware/HardwareBuffer;", false);
+	FImageResources_UScale = FindField(JEnv, FImageResourcesClass, "UScale", "F", false);
+	FImageResources_UOffset = FindField(JEnv, FImageResourcesClass, "UOffset", "F", false);
+	FImageResources_VScale = FindField(JEnv, FImageResourcesClass, "VScale", "F", false);
+	FImageResources_VOffset = FindField(JEnv, FImageResourcesClass, "VOffset", "F", false);
+
+	FImageResources_ReleaseFN = JEnv->GetMethodID(FImageResourcesClass, "Release", "()V");
+
+	// - - - - - - - - - - - - - - - - - - - - - - - - - - - - - - - - - - - - - - - - - - - - - - - - - - - - - - - - - - - - - - -
+
+	int32 SDKint = 0;
+	jclass localVersionClass = JEnv->FindClass("android/os/Build$VERSION");
+	if (localVersionClass)
+	{
+		jfieldID SdkIntFieldID = JEnv->GetStaticFieldID(localVersionClass, "SDK_INT", "I");
+		if (SdkIntFieldID)
+		{
+			SDKint = JEnv->GetStaticIntField(localVersionClass, SdkIntFieldID);
+		}
+		JEnv->DeleteLocalRef(localVersionClass);
+	}
+
+	// - - - - - - - - - - - - - - - - - - - - - - - - - - - - - - - - - - - - - - - - - - - - - - - - - - - - - - - - - - - - - - -
 
+	// Does the SDK support for KEY_ALLOW_FRAME_DROP in MediaFormat exist (and hence allow for throttle-free use of the Surface queue)?
+	if (SDKint < 31)
+	{
+		// No. Setup "read from surface" event to allow throttling
+		CodecSurfaceReadEvent = FPlatformProcess::GetSynchEventFromPool();
+		ensure(CodecSurfaceReadEvent != nullptr);
+		CodecSurfaceReadEvent->Trigger();
+	}
+	else
+	{
+		// Yes! No need for throttling...
+		CodecSurfaceReadEvent = nullptr;
+	}
+	
 #if ELECTRA_INIT_ON_RENDERTHREAD
 	// enqueue to RT to ensure GL resources are created on the appropriate thread.
 	SurfaceInitEvent = FPlatformProcess::GetSynchEventFromPool(true);
@@ -158,7 +334,7 @@ FElectraTextureSampleSupport::FElectraTextureSampleSupport()
 				{
 					JNIEnv* JEnv = FAndroidApplication::GetJavaEnv();
 					// Setup Java side of things
-					JEnv->CallVoidMethod(Object, InitializeFN.Method, FAndroidMisc::ShouldUseVulkan());
+					JEnv->CallVoidMethod(Object, InitializeFN.Method, UseGpuOutputPath, FAndroidMisc::ShouldUseVulkan(), (jlong)this);
 					// Query surface to be used for decoder
 					jobject Surface = JEnv->CallObjectMethod(Object, GetCodecSurfaceFN.Method);
 					CodecSurface = JEnv->NewGlobalRef(Surface);
@@ -199,6 +375,12 @@ FElectraTextureSampleSupport::~FElectraTextureSampleSupport()
 	CodecSurface = nullptr;
 	CodecSurfaceLock.Unlock();
 
+	if (CodecSurfaceReadEvent)
+	{
+		FPlatformProcess::ReturnSynchEventToPool(CodecSurfaceReadEvent);
+		CodecSurfaceReadEvent = nullptr;
+	}
+
 	if (IsInGameThread())
 	{
 		// enqueue to RT to ensure GL resources are released on the appropriate thread.
@@ -314,13 +496,25 @@ int32 FElectraTextureSampleSupport::GetFrameData(FElectraTextureSample* InTarget
 	{
 		if (InTargetSample)
 		{
-			jobject buffer = JEnv->GetObjectField(OutputInfo, FFrameUpdateInfo_Buffer);
-			if (buffer != nullptr)
+			if (UseGpuOutputPath)
 			{
-				const void* outPixels = JEnv->GetDirectBufferAddress(buffer);
-				const int32 outCount = JEnv->GetDirectBufferCapacity(buffer);
-				InTargetSample->SetupFromBuffer(outPixels, outCount);
-				JEnv->DeleteLocalRef(buffer);
+				jobject ImageResources = JEnv->GetObjectField(OutputInfo, FFrameUpdateInfo_ImageResources);
+				if (ImageResources)
+				{
+					InTargetSample->SetImageResources(ImageResources);
+					JEnv->DeleteLocalRef(ImageResources);
+				}
+			}
+			else
+			{
+				jobject buffer = JEnv->GetObjectField(OutputInfo, FFrameUpdateInfo_Buffer);
+				if (buffer != nullptr)
+				{
+					const void* outPixels = JEnv->GetDirectBufferAddress(buffer);
+					const int32 outCount = JEnv->GetDirectBufferCapacity(buffer);
+					InTargetSample->SetupFromBuffer(outPixels, outCount);
+					JEnv->DeleteLocalRef(buffer);
+				}
 			}
 		}
 
@@ -363,6 +557,35 @@ jobject FElectraTextureSampleSupport::GetCodecSurface()
 	return NewSurfaceHandle;
 }
 
+jobject FElectraTextureSampleSupport::ImageResources_GetHardwareBuffer(jobject ImageResources)
+{
+	JNIEnv* JEnv = FAndroidApplication::GetJavaEnv();
+	return JEnv->GetObjectField(ImageResources, FImageResources_HardwareBufferHandle);
+}
+
+
+void FElectraTextureSampleSupport::ImageResources_GetScaleOffset(jobject ImageResources, FVector2f& OutScale, FVector2f& OutOffset)
+{
+	JNIEnv* JEnv = FAndroidApplication::GetJavaEnv();
+	OutScale.X = JEnv->GetFloatField(ImageResources, FImageResources_UScale);
+	OutScale.Y = JEnv->GetFloatField(ImageResources, FImageResources_VScale);
+	OutOffset.X = JEnv->GetFloatField(ImageResources, FImageResources_UOffset);
+	OutOffset.Y = JEnv->GetFloatField(ImageResources, FImageResources_VOffset);
+}
+
+//-----------------------------------------------------------------------------
+//-----------------------------------------------------------------------------
+//-----------------------------------------------------------------------------
+
+JNI_METHOD void Java_com_epicgames_unreal_ElectraTextureSample_nativeSignalSurfaceReadEvent(JNIEnv* jenv, jobject thiz, jlong InParentHandle)
+{
+	auto Instance = reinterpret_cast<FElectraTextureSampleSupport*>(InParentHandle);
+	if (Instance)
+	{
+		Instance->SignalImageReaderSurfaceRead();
+	}
+}
+
 //-----------------------------------------------------------------------------
 //-----------------------------------------------------------------------------
 //-----------------------------------------------------------------------------
@@ -387,7 +610,14 @@ void FElectraTextureSample::Initialize(FVideoDecoderOutput* InVideoDecoderOutput
 	IElectraTextureSampleBase::Initialize(InVideoDecoderOutput);
 	VideoDecoderOutputAndroid = static_cast<FVideoDecoderOutputAndroid*>(InVideoDecoderOutput);
 
-	if (VideoDecoderOutputAndroid->GetOutputType() == FVideoDecoderOutputAndroid::EOutputType::DirectToSurfaceAsQueue)
+	ensure(VideoDecoderOutputAndroid->GetOutputType() == FVideoDecoderOutputAndroid::EOutputType::DirectToSurfaceAsQueue);
+
+	if (Support->UseGpuOutputPath)
+	{
+		Support->GetFrameData(this);
+		Texture = nullptr;
+	}
+	else
 	{
 		ENQUEUE_RENDER_COMMAND(InitTextureSample)([WeakThis{ AsWeak() }](FRHICommandListImmediate& RHICmdList) {
 			if (TSharedPtr<FElectraTextureSample, ESPMode::ThreadSafe> This = WeakThis.Pin())
@@ -415,6 +645,232 @@ void FElectraTextureSample::Initialize(FVideoDecoderOutput* InVideoDecoderOutput
 	}
 }
 
+FTextureRHIRef FElectraTextureSample::InitializeTextureOES(AHardwareBuffer* HardwareBuffer)
+{
+	check(HardwareBuffer);
+	check(IsInRenderingThread() || IsInRHIThread());
+	check(!FAndroidMisc::ShouldUseVulkan());
+
+#if 0 // TODO this needs to run on the OpenGlContext thread, so it probably will have to go into RHI before it can work
+
+	RHICmdList.EnqueueLambda([WeakThis{ AsWeak() }, &InDstTexture](FRHICommandListImmediate& CmdList)
+	{
+		QUICK_SCOPE_CYCLE_COUNTER(STAT_FRHICommandUpdateDecoderExternaTexture_Execute);
+		if (TSharedPtr<FElectraTextureSample, ESPMode::ThreadSafe> This = WeakThis.Pin())
+		{
+			if (This->ImageResources == nullptr)
+			{
+				This->Texture = nullptr;
+				return;
+			}
+
+			jobject HardwareBufferObj = This->Support->ImageResources_GetHardwareBuffer(This->ImageResources);
+			AHardwareBuffer* HardwareBuffer = AHardwareBuffer_fromHardwareBuffer(FAndroidApplication::GetJavaEnv(), HardwareBufferObj);
+
+			static PFNEGLGETNATIVECLIENTBUFFERANDROIDPROC EglGetNativeClientBufferANDROID = nullptr;
+			if (EglGetNativeClientBufferANDROID == nullptr)
+			{
+				EglGetNativeClientBufferANDROID = (PFNEGLGETNATIVECLIENTBUFFERANDROIDPROC)((void*)eglGetProcAddress("eglGetNativeClientBufferANDROID"));
+				ensure(EglGetNativeClientBufferANDROID != nullptr);
+			}
+      
+			// Ext. EGL_ANDROID_get_native_client_buffer
+			EGLClientBuffer NativeClientBuffer = EglGetNativeClientBufferANDROID(HardwareBuffer);
+			if (NativeClientBuffer == nullptr)
+			{
+				UE_LOG(LogElectraSamples, Warning, TEXT("Could not get native client buffer!"));
+				return;
+			}
+
+			static PFNEGLCREATEIMAGEKHRPROC EglCreateImageKHR = nullptr;
+			if (EglCreateImageKHR == nullptr)
+			{
+				EglCreateImageKHR = (PFNEGLCREATEIMAGEKHRPROC)((void*)eglGetProcAddress("eglCreateImageKHR"));
+				ensure(EglCreateImageKHR != nullptr);
+			}
+
+			// Ext. EGL_ANDROID_image_native_buffer
+			EGLImageKHR EglImage = EglCreateImageKHR(eglGetCurrentDisplay(), EGL_NO_CONTEXT, EGL_NATIVE_BUFFER_ANDROID, NativeClientBuffer, nullptr);
+			if (EglImage == 0)
+			{
+				UE_LOG(LogElectraSamples, Warning, TEXT("Could not create EGLimage from native client buffer! B=0x%x E=0x%x"), NativeClientBuffer, eglGetError());
+				return;
+			}
+
+			static PFNGLEGLIMAGETARGETTEXTURE2DOESPROC GlEGLImageTargetTexture2DOES = nullptr;
+			if (GlEGLImageTargetTexture2DOES == nullptr)
+			{
+				GlEGLImageTargetTexture2DOES = (PFNGLEGLIMAGETARGETTEXTURE2DOESPROC)((void*)eglGetProcAddress("glEGLImageTargetTexture2DOES"));
+				ensure(GlEGLImageTargetTexture2DOES != nullptr);
+			}
+
+			GlEGLImageTargetTexture2DOES(GL_TEXTURE_EXTERNAL_OES, EglImage);
+			check(eglGetError() == EGL_SUCCESS);
+			check(glGetError() == 0);
+
+			// Bind the RHI texture's GL texture to the external image we got data about now...
+			glBindTexture(GL_TEXTURE_EXTERNAL_OES, *reinterpret_cast<int32*>(This->Texture->GetNativeResource()));
+			check(glGetError() == 0);
+		
+		
+		}
+	});
+
+	FIntPoint Dim = VideoDecoderOutput->GetDim();
+
+	if (Texture.IsValid() && (Texture->GetSizeXY() == Dim || (Texture->GetFlags() & ETextureCreateFlags::External) == ETextureCreateFlags::External))
+	{
+		// The existing texture is just fine... (exists & same size OR external (size does not matter))
+		return;
+	}
+
+	const FRHITextureCreateDesc Desc =
+		FRHITextureCreateDesc::Create2D(TEXT("FElectraTextureSampleExternalTexture"), 1, 1, PF_R8G8B8A8)
+		.SetInitialState(ERHIAccess::SRVMask)
+		.SetFlags(ETextureCreateFlags::External);
+
+	return RHICreateTexture(Desc);
+
+#endif
+
+	return nullptr;
+}
+
+FTextureRHIRef FElectraTextureSample::InitializeTextureVulkan(AHardwareBuffer* HardwareBuffer)
+{
+	check(HardwareBuffer);
+	check(FAndroidMisc::ShouldUseVulkan());
+
+	IVulkanDynamicRHI* RHI = GetIVulkanDynamicRHI();
+	VkDevice Device = RHI->RHIGetVkDevice();
+	const VkAllocationCallbacks* AllocationCallbacks = RHI->RHIGetVkAllocationCallbacks();
+
+	AHardwareBuffer_Desc Desc;
+	AHardwareBuffer_describe(HardwareBuffer, &Desc);
+	check((Desc.usage & AHARDWAREBUFFER_USAGE_GPU_SAMPLED_IMAGE) != 0);
+
+	uint32 SizeX = Desc.width;
+	uint32 SizeY = Desc.height;
+	uint32 NumMips = 1;
+	uint32 NumSamples = 1;
+
+	ETextureCreateFlags Flags = ETextureCreateFlags::External;
+	const FClearValueBinding& ClearValueBinding = FClearValueBinding::Transparent;
+
+	static PFN_vkGetAndroidHardwareBufferPropertiesANDROID VkGetAndroidHardwareBufferPropertiesANDROID = nullptr;
+	if (!VkGetAndroidHardwareBufferPropertiesANDROID)
+	{
+		VkGetAndroidHardwareBufferPropertiesANDROID = (PFN_vkGetAndroidHardwareBufferPropertiesANDROID)RHI->RHIGetVkInstanceProcAddr("vkGetAndroidHardwareBufferPropertiesANDROID");
+		check(VkGetAndroidHardwareBufferPropertiesANDROID);
+	}
+
+	VkAndroidHardwareBufferFormatPropertiesANDROID HardwareBufferFormatProperties;
+	ZeroVulkanStruct(HardwareBufferFormatProperties, VK_STRUCTURE_TYPE_ANDROID_HARDWARE_BUFFER_FORMAT_PROPERTIES_ANDROID);
+
+	VkAndroidHardwareBufferPropertiesANDROID HardwareBufferProperties;
+	ZeroVulkanStruct(HardwareBufferProperties, VK_STRUCTURE_TYPE_ANDROID_HARDWARE_BUFFER_PROPERTIES_ANDROID);
+	HardwareBufferProperties.pNext = &HardwareBufferFormatProperties;
+
+	VkResult Result = VkGetAndroidHardwareBufferPropertiesANDROID(Device, HardwareBuffer, &HardwareBufferProperties);
+	check(Result == VK_SUCCESS);
+
+	VkExternalFormatANDROID ExternalFormat;
+	ZeroVulkanStruct(ExternalFormat, VK_STRUCTURE_TYPE_EXTERNAL_FORMAT_ANDROID);
+	ExternalFormat.externalFormat = HardwareBufferFormatProperties.externalFormat;
+
+	VkExternalMemoryImageCreateInfo ExternalMemoryImageCreateInfo;
+	ZeroVulkanStruct(ExternalMemoryImageCreateInfo, VK_STRUCTURE_TYPE_EXTERNAL_MEMORY_IMAGE_CREATE_INFO);
+	ExternalMemoryImageCreateInfo.handleTypes = VK_EXTERNAL_MEMORY_HANDLE_TYPE_ANDROID_HARDWARE_BUFFER_BIT_ANDROID;
+	ExternalMemoryImageCreateInfo.pNext = &ExternalFormat;
+
+	VkImageCreateInfo ImageCreateInfo;
+	ZeroVulkanStruct(ImageCreateInfo, VK_STRUCTURE_TYPE_IMAGE_CREATE_INFO);
+	ImageCreateInfo.imageType = VK_IMAGE_TYPE_2D;
+	ImageCreateInfo.format = VK_FORMAT_UNDEFINED;
+	ImageCreateInfo.extent.width = SizeX;
+	ImageCreateInfo.extent.height = SizeY;
+	ImageCreateInfo.extent.depth = 1;
+	ImageCreateInfo.mipLevels = NumMips;
+	ImageCreateInfo.arrayLayers = Desc.layers;
+
+	ImageCreateInfo.flags = 0;
+	ImageCreateInfo.samples = VK_SAMPLE_COUNT_1_BIT;
+
+	ImageCreateInfo.tiling = VK_IMAGE_TILING_OPTIMAL;
+	ImageCreateInfo.usage = VK_IMAGE_USAGE_SAMPLED_BIT;
+	ImageCreateInfo.sharingMode = VK_SHARING_MODE_EXCLUSIVE;
+	ImageCreateInfo.queueFamilyIndexCount = 0;
+	ImageCreateInfo.pQueueFamilyIndices = nullptr;
+	ImageCreateInfo.initialLayout = VK_IMAGE_LAYOUT_UNDEFINED;
+
+	ImageCreateInfo.pNext = &ExternalMemoryImageCreateInfo;
+
+	static PFN_vkCreateImage VkCreateImage = nullptr;
+	if (!VkCreateImage)
+	{
+		VkCreateImage = (PFN_vkCreateImage)RHI->RHIGetVkInstanceProcAddr("vkCreateImage");
+		check(VkCreateImage);
+	}
+
+	VkImage VulkanImage;
+	Result = VkCreateImage(Device, &ImageCreateInfo, nullptr, &VulkanImage);
+	check(Result == VK_SUCCESS);
+
+	VkMemoryDedicatedAllocateInfo MemoryDedicatedAllocateInfo;
+	ZeroVulkanStruct(MemoryDedicatedAllocateInfo, VK_STRUCTURE_TYPE_MEMORY_DEDICATED_ALLOCATE_INFO);
+	MemoryDedicatedAllocateInfo.image = VulkanImage;
+	MemoryDedicatedAllocateInfo.buffer = VK_NULL_HANDLE;
+
+	VkImportAndroidHardwareBufferInfoANDROID ImportAndroidHardwareBufferInfo;
+	ZeroVulkanStruct(ImportAndroidHardwareBufferInfo, VK_STRUCTURE_TYPE_IMPORT_ANDROID_HARDWARE_BUFFER_INFO_ANDROID);
+	ImportAndroidHardwareBufferInfo.buffer = HardwareBuffer;
+	ImportAndroidHardwareBufferInfo.pNext = &MemoryDedicatedAllocateInfo;
+
+	uint32 MemoryTypeBits = HardwareBufferProperties.memoryTypeBits;
+	check(MemoryTypeBits > 0); // No index available, this should never happen
+	uint32 MemoryTypeIndex = 0;
+	for (;(MemoryTypeBits & 1) != 1; ++MemoryTypeIndex)
+	{
+		MemoryTypeBits >>= 1;
+	}
+
+	VkMemoryAllocateInfo MemoryAllocateInfo;
+	ZeroVulkanStruct(MemoryAllocateInfo, VK_STRUCTURE_TYPE_MEMORY_ALLOCATE_INFO);
+	MemoryAllocateInfo.allocationSize = HardwareBufferProperties.allocationSize;
+	MemoryAllocateInfo.memoryTypeIndex = MemoryTypeIndex;
+	MemoryAllocateInfo.pNext = &ImportAndroidHardwareBufferInfo;
+
+	static PFN_vkAllocateMemory VkAllocateMemory = nullptr;
+	if (!VkAllocateMemory)
+	{
+		VkAllocateMemory = (PFN_vkAllocateMemory)RHI->RHIGetVkInstanceProcAddr("vkAllocateMemory");
+		check(VkAllocateMemory);
+	}
+
+	VkDeviceMemory VulkanDeviceMemory;
+	Result = VkAllocateMemory(Device, &MemoryAllocateInfo, AllocationCallbacks, &VulkanDeviceMemory);
+	check(Result == VK_SUCCESS);
+
+	static PFN_vkBindImageMemory VkBindImageMemory = nullptr;
+	if (!VkBindImageMemory)
+	{
+		VkBindImageMemory = (PFN_vkBindImageMemory)RHI->RHIGetVkInstanceProcAddr("vkBindImageMemory");
+		check(VkBindImageMemory);
+	}
+
+	Result = VkBindImageMemory(Device, VulkanImage, VulkanDeviceMemory, 0);
+	check(Result == VK_SUCCESS);
+
+	FElectraTextureSampleVulkanResources* ElectraTextureSampleVulkanResources = new FElectraTextureSampleVulkanResources{ VulkanImage, VulkanDeviceMemory };
+	FVulkanRHIExternalImageDeleteCallbackInfo ExternalImageDeleteCallbackInfo =
+	{
+		ElectraTextureSampleVulkanResources,
+		CleanupImageResourcesVulkan
+	};
+
+	return RHI->RHICreateTexture2DFromResource(PF_Unknown, SizeX, SizeY, NumMips, NumSamples, VulkanImage, Flags, ClearValueBinding, ExternalImageDeleteCallbackInfo);
+}
+
 
 void FElectraTextureSample::InitializeTexture(EPixelFormat PixelFormat)
 {
@@ -447,6 +903,24 @@ void FElectraTextureSample::InitializeTexture(EPixelFormat PixelFormat)
 	return;
 }
 
+void FElectraTextureSample::SetImageResources(jobject InImageResources)
+{
+	CleanupImageResources();
+	ImageResources = FAndroidApplication::GetJavaEnv()->NewGlobalRef(InImageResources);
+}
+
+#if !UE_SERVER
+void FElectraTextureSample::ShutdownPoolable()
+{
+	IElectraTextureSampleBase::ShutdownPoolable();
+//Q: this will trigger us to reuse the RHI/GL texture over and over -> sounds good. But is it 100%?
+//   - IF(!) the GL texture refs the EGL image internally we might keep the resources of that until a new one is assigned. Longer than needed! (although we free things up at the Java level (CleanupImageResources))
+// 	 - MIGHT that starve the decoder for buffers (we only have a max. number of "Images", too)
+//  >> set it to NULL and we'll recreate it... less ambigous maybe...
+//	Texture = nullptr;
+	CleanupImageResources();
+}
+#endif //!UE_SERVER
 
 void FElectraTextureSample::SetupFromBuffer(const void* InBuffer, int32 InBufferSize)
 {
@@ -465,6 +939,50 @@ void FElectraTextureSample::SetupFromBuffer(const void* InBuffer, int32 InBuffer
 	FMemory::Memcpy(Buffer, InBuffer, InBufferSize);
 }
 
+void FElectraTextureSample::CleanupImageResources()
+{
+	if (ImageResources)
+	{
+		jmethodID ReleaseFN = Support->GetImageResources_ReleaseFN();
+		jobject Resources = ImageResources;
+
+		if (IsInRHIThread())
+		{
+			CleanupImageResourcesJNI(Resources, ReleaseFN);
+		}
+		else if (IsInRenderingThread())
+		{
+			FRHICommandListExecutor::GetImmediateCommandList().EnqueueLambda([Resources, ReleaseFN](FRHICommandListImmediate&)
+			{
+				CleanupImageResourcesJNI(Resources, ReleaseFN);
+			});
+		}
+		else // Neither RHI nor Render thread
+		{
+			ENQUEUE_RENDER_COMMAND(ReleaseDecoderResources)([Resources, ReleaseFN](FRHICommandListImmediate& RHICmdList)
+			{
+				RHICmdList.EnqueueLambda([Resources, ReleaseFN](FRHICommandListImmediate&)
+				{
+					CleanupImageResourcesJNI(Resources, ReleaseFN);
+				});
+			});
+		}
+
+		ImageResources = nullptr;
+	}
+}
+
+
+FElectraTextureSample::~FElectraTextureSample()
+{
+	if (Buffer)
+	{
+		FMemory::Free(Buffer);
+	}
+
+	CleanupImageResources();
+}
+
 
 EMediaTextureSampleFormat FElectraTextureSample::GetFormat() const
 {
@@ -483,7 +1001,184 @@ uint32 FElectraTextureSample::GetStride() const
 }
 
 
+void  FElectraTextureSample::CopyFromExternalTextureOES(FRHICommandListImmediate& RHICmdList, FTextureRHIRef& InDstTexture, FTextureRHIRef& InSrcTexture, const FVector2f& InScale, const FVector2f& InOffset)
+{
+	FLinearColor Offset = { InOffset.X, InOffset.Y, 0.f, 0.f };
+	FLinearColor ScaleRotation = { InScale.X, 0.f,
+								  0.f, InScale.Y };
+
+	FGraphicsPipelineStateInitializer GraphicsPSOInit;
+	FRHITexture* RenderTarget = InDstTexture.GetReference();
+
+	RHICmdList.Transition(FRHITransitionInfo(InDstTexture, ERHIAccess::Unknown, ERHIAccess::RTV));
+
+	FRHIRenderPassInfo RPInfo(RenderTarget, ERenderTargetActions::DontLoad_Store);
+	RHICmdList.BeginRenderPass(RPInfo, TEXT("ConvertMedia_ExternalTexture"));
+	{
+		const FIntPoint OutputDim = GetOutputDim();
+
+		RHICmdList.ApplyCachedRenderTargets(GraphicsPSOInit);
+		RHICmdList.SetViewport(0, 0, 0.0f, (float)OutputDim.X, (float)OutputDim.Y, 1.0f);
+
+		GraphicsPSOInit.DepthStencilState = TStaticDepthStencilState<false, CF_Always>::GetRHI();
+		GraphicsPSOInit.RasterizerState = TStaticRasterizerState<>::GetRHI();
+		GraphicsPSOInit.BlendState = TStaticBlendStateWriteMask<CW_RGBA, CW_NONE, CW_NONE, CW_NONE, CW_NONE, CW_NONE, CW_NONE, CW_NONE>::GetRHI();
+		GraphicsPSOInit.PrimitiveType = PT_TriangleStrip;
+
+		// configure media shaders
+		auto ShaderMap = GetGlobalShaderMap(GMaxRHIFeatureLevel);
+		TShaderMapRef<FMediaShadersVS> VertexShader(ShaderMap);
+
+		GraphicsPSOInit.BoundShaderState.VertexDeclarationRHI = GMediaVertexDeclaration.VertexDeclarationRHI;
+		GraphicsPSOInit.BoundShaderState.VertexShaderRHI = VertexShader.GetVertexShader();
+
+		FSamplerStateInitializerRHI InSamplerStateInitializer(SF_Bilinear, AM_Clamp, AM_Clamp, AM_Clamp);
+		FSamplerStateRHIRef InSamplerState = RHICreateSamplerState(InSamplerStateInitializer);
+
+		TShaderMapRef<FReadTextureExternalPS> CopyShader(ShaderMap);
+		GraphicsPSOInit.BoundShaderState.PixelShaderRHI = CopyShader.GetPixelShader();
+		SetGraphicsPipelineState(RHICmdList, GraphicsPSOInit, 0);
+		SetShaderParametersLegacyPS(RHICmdList, CopyShader, InSrcTexture, InSamplerState, ScaleRotation, Offset);
+
+		// draw full size quad into render target
+		FBufferRHIRef VertexBuffer = CreateTempMediaVertexBuffer();
+		RHICmdList.SetStreamSource(0, VertexBuffer, 0);
+		// set viewport to RT size
+		RHICmdList.SetViewport(0, 0, 0.0f, (float)OutputDim.X, (float)OutputDim.Y, 1.0f);
+
+		RHICmdList.DrawPrimitive(0, 2, 1);
+	}
+	RHICmdList.EndRenderPass();
+
+	RHICmdList.Transition(FRHITransitionInfo(InDstTexture, ERHIAccess::RTV, ERHIAccess::SRVMask));
+}
+
+void  FElectraTextureSample::CopyFromExternalTextureVulkan(FRHICommandListImmediate& RHICmdList, FTextureRHIRef& InDstTexture, FTextureRHIRef& InSrcTexture, const FVector2f& InScale, const FVector2f& InOffset)
+{
+	FGraphicsPipelineStateInitializer GraphicsPSOInit;
+	FRHITexture* RenderTarget = InDstTexture;
+
+	RHICmdList.Transition(FRHITransitionInfo(InDstTexture, ERHIAccess::Unknown, ERHIAccess::RTV));
+
+	FRHIRenderPassInfo RPInfo(RenderTarget, ERenderTargetActions::DontLoad_Store);
+	RHICmdList.BeginRenderPass(RPInfo, TEXT("ConvertMedia"));
+	{
+		const FIntPoint OutputDim = GetOutputDim();
+
+		RHICmdList.ApplyCachedRenderTargets(GraphicsPSOInit);
+		RHICmdList.SetViewport(0, 0, 0.0f, (float)OutputDim.X, (float)OutputDim.Y, 1.0f);
+
+		GraphicsPSOInit.DepthStencilState = TStaticDepthStencilState<false, CF_Always>::GetRHI();
+		GraphicsPSOInit.RasterizerState = TStaticRasterizerState<>::GetRHI();
+		GraphicsPSOInit.BlendState = TStaticBlendStateWriteMask<CW_RGBA, CW_NONE, CW_NONE, CW_NONE, CW_NONE, CW_NONE, CW_NONE, CW_NONE>::GetRHI();
+		GraphicsPSOInit.PrimitiveType = PT_TriangleStrip;
+
+		// configure media shaders
+		auto ShaderMap = GetGlobalShaderMap(GMaxRHIFeatureLevel);
+		TShaderMapRef<FMediaShadersVS> VertexShader(ShaderMap);
+
+		GraphicsPSOInit.BoundShaderState.VertexDeclarationRHI = GMediaVertexDeclaration.VertexDeclarationRHI;
+		GraphicsPSOInit.BoundShaderState.VertexShaderRHI = VertexShader.GetVertexShader();
+
+		auto YUVMtx = GetSampleToRGBMatrix();
+		FMatrix44f ColorSpaceMtx;
+		//------------------------------------------------------------------------
+		//GetColorSpaceConversionMatrixForSample(this, ColorSpaceMtx);
+		{
+			const UE::Color::FColorSpace& Working = UE::Color::FColorSpace::GetWorking();
+	
+			if (GetMediaTextureSampleColorConverter())
+			{
+				ColorSpaceMtx = FMatrix44f::Identity;
+			}
+			else
+			{
+				ColorSpaceMtx = FMatrix44f(Working.GetXYZToRgb().GetTransposed() * GetGamutToXYZMatrix());
+			}
+	
+			float NF = GetHDRNitsNormalizationFactor();
+			if (NF != 1.0f)
+			{
+				ColorSpaceMtx = ColorSpaceMtx.ApplyScale(NF);
+			}
+		}
+		//------------------------------------------------------------------------
+
+		TShaderMapRef<FVYUConvertPS> ConvertShader(ShaderMap);
+		GraphicsPSOInit.BoundShaderState.PixelShaderRHI = ConvertShader.GetPixelShader();
+		SetGraphicsPipelineState(RHICmdList, GraphicsPSOInit, 0);
+		SetShaderParametersLegacyPS(RHICmdList, ConvertShader, InSrcTexture, OutputDim, YUVMtx, GetEncodingType(), ColorSpaceMtx);
+
+
+		// draw full size quad into render target
+		FBufferRHIRef VertexBuffer = CreateTempMediaVertexBuffer(); 
+		RHICmdList.SetStreamSource(0, VertexBuffer, 0);
+		// set viewport to RT size
+		RHICmdList.SetViewport(0, 0, 0.0f, (float)OutputDim.X, (float)OutputDim.Y, 1.0f);
+
+		RHICmdList.DrawPrimitive(0, 2, 1);
+	}
+	RHICmdList.EndRenderPass();
+
+	RHICmdList.Transition(FRHITransitionInfo(RenderTarget, ERHIAccess::RTV, ERHIAccess::SRVGraphics));
+}
+
+
 bool FElectraTextureSample::Convert(FRHICommandListImmediate& RHICmdList, FTextureRHIRef& InDstTexture, const FConversionHints& Hints)
+{
+	if (Support->UseGpuOutputPath)
+		return ConvertGpuOutputPath(RHICmdList, InDstTexture, Hints);
+	else
+		return ConvertCpuOutputPath(RHICmdList, InDstTexture, Hints);
+}
+
+bool FElectraTextureSample::ConvertGpuOutputPath(FRHICommandListImmediate& RHICmdList, FTextureRHIRef& InDstTexture, const FConversionHints& Hints)
+{
+	check(IsInRenderingThread());
+
+	if (GDynamicRHI->RHIIsRenderingSuspended() || ImageResources == nullptr)
+	{
+		return false;
+	}
+
+	FTextureRHIRef SrcTexture;
+	jobject HardwareBufferObj = Support->ImageResources_GetHardwareBuffer(ImageResources);
+	AHardwareBuffer* HardwareBuffer = AHardwareBuffer_fromHardwareBuffer(FAndroidApplication::GetJavaEnv(), HardwareBufferObj);
+	ensure(HardwareBuffer);
+
+	check(VideoDecoderOutputAndroid->GetOutputType() == FVideoDecoderOutputAndroid::EOutputType::DirectToSurfaceAsQueue);
+	if (FAndroidMisc::ShouldUseVulkan())
+	{
+		SrcTexture = InitializeTextureVulkan(HardwareBuffer);
+	}
+	else
+	{
+		SrcTexture = InitializeTextureOES(HardwareBuffer);
+	}
+
+	ensure(SrcTexture);
+	if (FAndroidMisc::ShouldUseVulkan())
+	{
+		FVector2f Scale, Offset;
+		Support->ImageResources_GetScaleOffset(ImageResources, Scale, Offset);
+
+		CopyFromExternalTextureVulkan(RHICmdList, InDstTexture, SrcTexture, Scale, Offset);
+	}
+	else
+	{
+		FVector2f Scale, Offset;
+		Support->ImageResources_GetScaleOffset(ImageResources, Scale, Offset);
+
+		//FOR NOW(?) THIS IS DONE HERE TO MAKE SURE WE HAVE EASY ACCESS TO THE SCALE/OFFSET/ROTATION VALUES FOR EACH SAMPLE
+		//(the code using a map & GUID lookup assumes ONE "current" value per player... which does entirely NOT work in reality (queue of frames))
+		CopyFromExternalTextureOES(RHICmdList, InDstTexture, SrcTexture, Scale, Offset);
+	}
+
+	return true;
+}
+
+
+bool FElectraTextureSample::ConvertCpuOutputPath(FRHICommandListImmediate& RHICmdList, FTextureRHIRef& InDstTexture, const FConversionHints& Hints)
 {
 	if (GDynamicRHI->RHIIsRenderingSuspended())
 	{
diff --git a/Engine/Plugins/Media/ElectraUtil/Source/ElectraSamples/Private/Java/src/com/epicgames/unreal/ElectraTextureSample.java b/Engine/Plugins/Media/ElectraUtil/Source/ElectraSamples/Private/Java/src/com/epicgames/unreal/ElectraTextureSample.java
--- a/Engine/Plugins/Media/ElectraUtil/Source/ElectraSamples/Private/Java/src/com/epicgames/unreal/ElectraTextureSample.java
+++ b/Engine/Plugins/Media/ElectraUtil/Source/ElectraSamples/Private/Java/src/com/epicgames/unreal/ElectraTextureSample.java
@@ -12,11 +12,33 @@ import android.opengl.*;
 import android.graphics.SurfaceTexture;
 import android.view.Surface;
 import java.nio.ByteBuffer;
-
+import android.media.ImageReader;
+import android.media.Image;
+import android.hardware.HardwareBuffer;
 
 public class ElectraTextureSample
 {
-	private static final String TAG = "[ElectraPlayerSDK]";
+	private static final String TAG = "[ElectraTextureSample]";
+
+	class FImageResources
+	{
+		public void Release()
+		{
+			if (Image != null)
+			{
+				Image.close();
+				Image = null;
+			}
+		}
+
+		public android.media.Image Image = null;
+		public android.hardware.HardwareBuffer HardwareBuffer = null;
+
+		public float UScale = 0.0f;
+		public float UOffset = 0.0f;
+		public float VScale = 0.0f;
+		public float VOffset = 0.0f;
+	}
 
 	public class FFrameUpdateInfo
 	{
@@ -30,6 +52,7 @@ public class ElectraTextureSample
 		public float VScale = 0.0f;
 		public float VOffset = 0.0f;
 		public int NumPending = 0;
+		public FImageResources ImageResources = null;
 	}
 
 	private FBitmapRenderer BitmapRenderer = null;
@@ -38,9 +61,12 @@ public class ElectraTextureSample
 	{
 	}
 
-	public void Initialize(boolean bVulkan)
+	public void Initialize(boolean bNewRenderer, boolean bVulkan, long parentHandle)
 	{
-		BitmapRenderer = new FBitmapRenderer(bVulkan);
+		if (bNewRenderer)
+			BitmapRenderer = new FBitmapRendererNew(bVulkan, parentHandle, 16);
+		else
+			BitmapRenderer = new FBitmapRendererOld(bVulkan);
 	}
 
 	public void Release()
@@ -75,10 +101,17 @@ public class ElectraTextureSample
 	/***************************************************************************************************************************************************/
 	/***************************************************************************************************************************************************/
 
+	abstract class FBitmapRenderer
+	{
+		public abstract void release();
+		public abstract FFrameUpdateInfo GetVideoFrameUpdateInfo(int destTexture, int width, int height, boolean bIs10Bit);
+		public abstract android.view.Surface getSurface();
+	}
+
 	/*
 		All this internal surface view does is manage the offscreen bitmap that the media player decoding can render into for eventual extraction to the UE4 buffers.
 	*/
-	class FBitmapRenderer
+	class FBitmapRendererOld extends FBitmapRenderer
 	{
 		private java.nio.Buffer mFrameData = null;
 		private int mFrameDataSize = 0;
@@ -110,7 +143,7 @@ public class ElectraTextureSample
 
 		public int mNativeDecoderID;
 
-		public FBitmapRenderer(boolean vulkanRenderer)
+		public FBitmapRendererOld(boolean vulkanRenderer)
 		{
 			mVulkanRenderer = vulkanRenderer;
 
@@ -940,7 +973,223 @@ public class ElectraTextureSample
 		}
 	}
 
+	/*
+		All this internal surface view does is manage the offscreen bitmap that the media player decoding can render into for eventual extraction to the UE4 buffers.
+	*/
+	class FBitmapRendererNew extends FBitmapRenderer implements ImageReader.OnImageAvailableListener
+	{
+		private java.nio.Buffer mFrameData = null;
+		private int mFrameDataSize = 0;
+		private android.media.ImageReader mImageReader = null;
+		private android.os.HandlerThread mHandlerThread = null;
+		private android.os.Handler mImageReaderHandler = null;
+		private java.util.concurrent.BlockingQueue<Image> mImageQueue = null;
+		private int mTextureWidth = -1;
+		private int mTextureHeight = -1;
+		private android.view.Surface mSurface = null;
+		private boolean mUseOwnContext = true;
+		private boolean mVulkanRenderer = false;
+		private long mParentHandle = 0;
+
+		public int mNativeDecoderID;
+
+		public FBitmapRendererNew(boolean vulkanRenderer, long parentHandle, int maxQueueLength)
+		{
+			mVulkanRenderer = vulkanRenderer;
+			mParentHandle = parentHandle;
+
+			mImageReader = android.media.ImageReader.newInstance(1, 1, android.graphics.ImageFormat.YUV_420_888, maxQueueLength, android.hardware.HardwareBuffer.USAGE_GPU_SAMPLED_IMAGE);
+
+			mSurface = mImageReader.getSurface();
+
+			mImageQueue = new java.util.concurrent.LinkedBlockingQueue<>(maxQueueLength / 2);
+			// We need to spawn our own thread here, because the main looper cannot be relied on to collect images before the ImageReader queue overflows
+			 mImageReaderHandler = new android.os.Handler(android.os.Looper.getMainLooper());
+			//mHandlerThread = new android.os.HandlerThread("ElectraTextureSampleHandlerThread");
+			//mHandlerThread.start();
+			//mImageReaderHandler = new android.os.Handler(mHandlerThread.getLooper());
+			mImageReader.setOnImageAvailableListener(this, mImageReaderHandler);
+		}
+
+		public void onImageAvailable(ImageReader reader)
+		{
+			// Get image data from Surface / Queue...
+			android.media.Image Image = null;
+			try {
+				Image =  mImageReader.acquireNextImage();
+			} catch(Exception e) {
+				Log.d(TAG, String.format("ImageReader.acquireNextImage() threw: %s", e.toString()));
+			}
+
+			if (Image != null)
+			{
+				// Signal that we received the image
+				// (used on pre API 31 to make sure we only have one image in the queue; NoOp later on!)
+				nativeSignalSurfaceReadEvent(mParentHandle);
+
+				// Deliver new image into our own queue
+				try {
+					mImageQueue.put(Image);
+				} catch(InterruptedException e) {
+					// [...] ?
+				}
+			}
+		}
+
+		public boolean isValid()
+		{
+			return mImageReader != null;
+		}
+
+		public android.view.Surface getSurface()
+		{
+			return mSurface;
+		}
+
+		public FFrameUpdateInfo GetVideoFrameUpdateInfo(int destTexture, int width, int height, boolean bIs10Bit)
+		{
+			mTextureWidth = width;
+			mTextureHeight = height;
+
+			return updateFrameData(bIs10Bit);
+		}
+
+		public FFrameUpdateInfo updateFrameData(boolean bIs10Bit)
+		{
+			synchronized(this)
+			{
+				// Copy surface texture to destination texture.
+				FImageResources ImageResources = getFrameTextureInfo();
+				if (ImageResources == null)
+				{
+					return null;
+				}
+
+				// Get associated frame metadata
+				FFrameUpdateInfo frameUpdateInfo = new FFrameUpdateInfo();
+
+				frameUpdateInfo.ImageResources = ImageResources;
+
+				frameUpdateInfo.Timestamp = -1;
+				frameUpdateInfo.bFrameReady = true;
+				frameUpdateInfo.bRegionChanged = false;
+
+				return frameUpdateInfo;
+			}
+		}
+
+
+		private FImageResources getFrameTextureInfo()
+		{
+			if (null == mImageReader)
+			{
+				// Can't update if there's no surface to update into.
+				return null;
+			}
+
+			// Prepare to capture any resources we need to track for the output image
+			FImageResources ImageResources = new FImageResources();
+
+			// Grab the next image in the queue
+			try {
+				ImageResources.Image = mImageQueue.take();
+			} catch(InterruptedException e) {
+				// [...] ?
+			}
+
+			if (ImageResources.Image == null)
+			{
+				return null;
+			}
+
+			ImageResources.HardwareBuffer = ImageResources.Image.getHardwareBuffer();
+			if (ImageResources.HardwareBuffer == null)
+			{
+				Log.d(TAG, "Could not get HardwareBuffer from Image!");
+				return null;
+			}
+
+/*
+//ALL THIS COULD GO 100% NATIVE POSSIBLY - ONCE WE OWN THE Image WE SHOULD BE GOOD?
+			// Grab the EGLNativeClientBuffer from the Image using native code
+
+			if (mVulkanRenderer)
+			{
+				ImageResources.EglImageNativeHandle = 0;
+				ImageResources.HardwareBuffer = ImageResources.Image.getHardwareBuffer();
+				if (ImageResources.HardwareBuffer == null)
+				{
+					Log.d(TAG, "Could not get HardwareBuffer from Image!");
+					return null;
+				}
+			}
+			else
+			{
+				ImageResources.HardwareBuffer = null;
+				long NativeClientBuffer = nativeGetEGLNativeClientBuffer(ImageResources.Image.getHardwareBuffer());
+				if (NativeClientBuffer < 0)
+				{
+					Log.d(TAG, "Could not get native client buffer!");
+					return null;
+				}
+
+				ImageResources.EglImageNativeHandle = nativeCreateEGLImageKHR(EGL14.eglGetCurrentDisplay().getNativeHandle(), EGL14.eglGetCurrentContext().getNativeHandle(), NativeClientBuffer);
+				if (ImageResources.EglImageNativeHandle == 0)
+				{
+					Log.d(TAG, String.format("Could not create EGLimage from native client buffer! B=0x%x E=0x%x", NativeClientBuffer, EGL14.eglGetError()));
+					return null;
+				}
+
+				Log.d(TAG, String.format("NATIVE HANDLE: 0x%x", ImageResources.EglImageNativeHandle));
+			}
+*/
+			android.graphics.Rect CropRect = ImageResources.Image.getCropRect();
+			ImageResources.UScale = (float)(CropRect.right - CropRect.left) / mTextureWidth;
+			ImageResources.UOffset = (float)CropRect.left / mTextureWidth;
+			ImageResources.VScale = (float)(CropRect.bottom - CropRect.top) / mTextureHeight;
+			ImageResources.VOffset = (float)CropRect.top / mTextureHeight;
+
+			return ImageResources;
+		}
+
+		public FFrameUpdateInfo GetVideoFrameUpdateInfo(int width, int height, boolean bIs10Bit)
+		{
+			mTextureWidth = width;
+			mTextureHeight = height;
+
+			return updateFrameData(bIs10Bit);
+		}
+
+		public void release()
+		{
+			synchronized(this)
+			{
+				if (mImageReader != null)
+				{
+					mImageReader.setOnImageAvailableListener(null, null);
+				}
+			}
+
+			if (null != mSurface)
+			{
+				mSurface.release();
+				mSurface = null;
+			}
+			if (mImageReader != null)
+			{
+				mImageReader.close();
+				mImageReader = null;
+			}
+			if (mHandlerThread != null)
+			{
+				mHandlerThread.stop();
+				mHandlerThread = null;
+			}
+		}
+	}
+
 	public native void nativeClearCachedAttributeState(int PositionAttrib, int TexCoordsAttrib);
+	public native void nativeSignalSurfaceReadEvent(long parentHandle);
 }
 
 
diff --git a/Engine/Plugins/Media/ElectraUtil/Source/ElectraSamples/Public/Android/ElectraTextureSample.h b/Engine/Plugins/Media/ElectraUtil/Source/ElectraSamples/Public/Android/ElectraTextureSample.h
--- a/Engine/Plugins/Media/ElectraUtil/Source/ElectraSamples/Public/Android/ElectraTextureSample.h
+++ b/Engine/Plugins/Media/ElectraUtil/Source/ElectraSamples/Public/Android/ElectraTextureSample.h
@@ -15,10 +15,9 @@
 #include "MediaSamples.h"
 #include "MediaVideoDecoderOutputAndroid.h"
 
-#include "RHI.h"
-#include "RHIUtilities.h"
-
 class FElectraTextureSampleSupport;
+class FEvent;
+struct AHardwareBuffer;
 
 /**
  * Texture sample generated by AndroidMedia player.
@@ -33,6 +32,7 @@ public:
 	/** Default constructor. */
 	FElectraTextureSample(TSharedPtr<FElectraTextureSampleSupport, ESPMode::ThreadSafe> InSupport)
 	: Support(InSupport)
+	, ImageResources(nullptr)
 	, Texture(nullptr)
 	, Buffer(nullptr)
 	, BufferSize(0)
@@ -40,13 +40,11 @@ public:
 	}
 
 	/** Virtual destructor. */
-	virtual ~FElectraTextureSample()
-	{
-		if (Buffer)
-		{
-			FMemory::Free(Buffer);
-		}
-	}
+	virtual ~FElectraTextureSample();
+
+#if !UE_SERVER
+	virtual void ShutdownPoolable() override;
+#endif
 
 	void Initialize(FVideoDecoderOutput* InVideoDecoderOutput);
 
@@ -74,10 +72,20 @@ public:
 private:
 	friend class FElectraTextureSampleSupport;
 
+	FTextureRHIRef InitializeTextureOES(AHardwareBuffer* HardwareBuffer);
+	FTextureRHIRef InitializeTextureVulkan(AHardwareBuffer* HardwareBuffer);
+
 	void InitializeTexture(EPixelFormat PixelFormat);
 	void SetupFromBuffer(const void* InBuffer, int32 InBufferSize);
+	void SetImageResources(jobject InImageResources);
+	void CleanupImageResources();
+	void CopyFromExternalTextureOES(FRHICommandListImmediate& RHICmdList, FTextureRHIRef& InDstTexture, FTextureRHIRef& InSrcTexture, const FVector2f& Scale, const FVector2f& Offset);
+	void CopyFromExternalTextureVulkan(FRHICommandListImmediate& RHICmdList, FTextureRHIRef& InDstTexture, FTextureRHIRef& InSrcTexture, const FVector2f& InScale, const FVector2f& InOffset);
 
 	virtual bool Convert(FRHICommandListImmediate& RHICmdList, FTextureRHIRef& InDstTexture, const FConversionHints& Hints) override;
+	bool ConvertCpuOutputPath(FRHICommandListImmediate& RHICmdList, FTextureRHIRef& InDstTexture, const FConversionHints& Hints);
+	bool ConvertGpuOutputPath(FRHICommandListImmediate& RHICmdList, FTextureRHIRef& InDstTexture, const FConversionHints& Hints);
+
 	virtual uint32 GetConverterInfoFlags() const
 	{
 		return ConverterInfoFlags_Default;
@@ -87,11 +95,15 @@ private:
 
 	FVideoDecoderOutputAndroid* VideoDecoderOutputAndroid;
 
+	bool UseGpuOutputPath;
+
 	void* Buffer;
 	int32 BufferSize;
 
 	/** Texture resource. */
 	TRefCountPtr<FRHITexture> Texture;
+
+	jobject ImageResources;
 };
 
 
diff --git a/Engine/Shaders/Private/MediaShaders.usf b/Engine/Shaders/Private/MediaShaders.usf
--- a/Engine/Shaders/Private/MediaShaders.usf
+++ b/Engine/Shaders/Private/MediaShaders.usf
@@ -517,15 +517,17 @@ void RGB8toY8ConvertPS(float4 InPosition : SV_POSITION,
 }
 
 
+TextureExternal InExternalTexture;
+
 // shader to read from TextureExternal source
 void ReadTextureExternalPS(
 	float4 InPosition : SV_POSITION,
 	float2 InUV : TEXCOORD0,
 	out float4 OutColor : SV_Target0)
 {
-    float2 UV = float2(dot(InUV, ReadTextureExternalUB.ScaleRotation.xy), dot(InUV, ReadTextureExternalUB.ScaleRotation.zw)) + ReadTextureExternalUB.Offset;
-	float4 InColor = ReadTextureExternalUB.Texture.Sample(ReadTextureExternalUB.SamplerP, UV);
-    OutColor = float4(sRGBToLinear(InColor.xyz), InColor.w);
+	float2 UV = float2(dot(InUV, ReadTextureExternalUB.ScaleRotation.xy), dot(InUV, ReadTextureExternalUB.ScaleRotation.zw)) + ReadTextureExternalUB.Offset;
+	float4 RGBOut = TextureExternalSample(InExternalTexture, ReadTextureExternalUB.SamplerP, UV).gbra;
+	OutColor = float4(sRGBToLinear(RGBOut.xyz), 1.0);
 }
 
 /* Media Capture Pixel shaders 
@@ -586,6 +588,18 @@ void RGB10toYUVv210ConvertPS(
 	OutColor.xyzw = uint4(W0, W1, W2, W3);
 }
 
+// Shader to convert from Android's VYU to RGBA
+void VYUConvertPS(
+	float4 InPosition : SV_POSITION,
+	float2 InUV : TEXCOORD0,
+	out float4 OutColor : SV_Target0)
+{
+	float2 ScaledUV = InUV * VYUConvertUB.UVScale;
+	float3 YUV = VYUConvertUB.Texture.Sample(VYUConvertUB.Sampler, ScaledUV).gbr;
+
+	// Output with solid alpha
+	OutColor = float4(ApplyColorProcessing(YUV, VYUConvertUB.ColorTransform, VYUConvertUB.CSTransform, VYUConvertUB.EOTF), 1.0f); 
+}
 
 // Following defines must reflect values in EMediaCaptureConversionOperation in MediaOutput
 #define INVERT_ALPHA 3
diff --git a/Engine/Source/Runtime/Media/Public/IMediaTextureSample.h b/Engine/Source/Runtime/Media/Public/IMediaTextureSample.h
--- a/Engine/Source/Runtime/Media/Public/IMediaTextureSample.h
+++ b/Engine/Source/Runtime/Media/Public/IMediaTextureSample.h
@@ -122,6 +122,9 @@ enum class EMediaTextureSampleFormat
 
 	/** ARGB 16-bit per component, big endian */
 	ARGB16_BIG,
+
+	/** External texture with 4 components (VYUX) per texel. */
+	ExternalVYU,
 };
 
 namespace MediaTextureSampleFormat
diff --git a/Engine/Source/Runtime/MediaAssets/Private/Misc/MediaTextureResource.cpp b/Engine/Source/Runtime/MediaAssets/Private/Misc/MediaTextureResource.cpp
--- a/Engine/Source/Runtime/MediaAssets/Private/Misc/MediaTextureResource.cpp
+++ b/Engine/Source/Runtime/MediaAssets/Private/Misc/MediaTextureResource.cpp
@@ -1326,6 +1326,19 @@ void FMediaTextureResource::GetColorSpaceConversionMatrixForSample(const TShared
 				}
 				break;
 
+				case EMediaTextureSampleFormat::ExternalVYU:
+				{
+					auto YUVMtx = Sample->GetSampleToRGBMatrix();
+					FMatrix44f ColorSpaceMtx;
+					GetColorSpaceConversionMatrixForSample(Sample, ColorSpaceMtx);
+
+					TShaderMapRef<FVYUConvertPS> ConvertShader(ShaderMap);
+					GraphicsPSOInit.BoundShaderState.PixelShaderRHI = ConvertShader.GetPixelShader();
+					SetGraphicsPipelineState(RHICmdList, GraphicsPSOInit, 0);
+					SetShaderParametersLegacyPS(RHICmdList, ConvertShader, InputTexture, OutputDim, YUVMtx, Sample->GetEncodingType(), ColorSpaceMtx);
+				}
+				break;
+
 				default:
 				{
 					// This should not happen in normal use: still - end the render pass to avoid any trouble with RHI
diff --git a/Engine/Source/Runtime/RenderCore/Private/MediaShaders.cpp b/Engine/Source/Runtime/RenderCore/Private/MediaShaders.cpp
--- a/Engine/Source/Runtime/RenderCore/Private/MediaShaders.cpp
+++ b/Engine/Source/Runtime/RenderCore/Private/MediaShaders.cpp
@@ -852,7 +852,6 @@ void FRGB8toY8ConvertPS::SetParameters(FRHIBatchedShaderParameters& BatchedParam
  *****************************************************************************/
 
 BEGIN_GLOBAL_SHADER_PARAMETER_STRUCT(FReadTextureExternalUB, )
-SHADER_PARAMETER_TEXTURE(Texture2D, Texture)
 SHADER_PARAMETER_SAMPLER(SamplerState, SamplerP)
 SHADER_PARAMETER(FLinearColor, ScaleRotation)
 SHADER_PARAMETER(FVector2f, Offset)
@@ -867,13 +866,13 @@ void FReadTextureExternalPS::SetParameters(FRHIBatchedShaderParameters& BatchedP
 	FReadTextureExternalUB UB;
 	{
 		UB.SamplerP = SamplerState;
-		UB.Texture = TextureExt;
 		UB.ScaleRotation = ScaleRotation;
 		UB.Offset = FVector2f(Offset.R, Offset.G);
 	}
 
 	TUniformBufferRef<FReadTextureExternalUB> Data = TUniformBufferRef<FReadTextureExternalUB>::CreateUniformBufferImmediate(UB, UniformBuffer_SingleFrame);
 	SetUniformBufferParameter(BatchedParameters, GetUniformBufferParameter<FReadTextureExternalUB>(), Data);
+	SetTextureParameter(BatchedParameters, InTexture, TextureExt);
 }
 
 /* FRGB8toUYVY8ConvertPS shader
@@ -940,3 +939,34 @@ FModifyAlphaSwizzleRgbaPS::FParameters* FModifyAlphaSwizzleRgbaPS::AllocateAndSe
 }
 
 
+/* FVYUConvertPS shader
+ *****************************************************************************/
+
+BEGIN_GLOBAL_SHADER_PARAMETER_STRUCT(FVYUConvertUB, )
+SHADER_PARAMETER(FMatrix44f, ColorTransform)
+SHADER_PARAMETER(FMatrix44f, CSTransform)
+SHADER_PARAMETER(uint32, EOTF)
+SHADER_PARAMETER(FVector2f, UVScale)
+SHADER_PARAMETER_TEXTURE(Texture2D, Texture)
+SHADER_PARAMETER_SAMPLER(SamplerState, Sampler)
+END_GLOBAL_SHADER_PARAMETER_STRUCT()
+
+IMPLEMENT_GLOBAL_SHADER_PARAMETER_STRUCT(FVYUConvertUB, "VYUConvertUB");
+IMPLEMENT_SHADER_TYPE(, FVYUConvertPS, TEXT("/Engine/Private/MediaShaders.usf"), TEXT("VYUConvertPS"), SF_Pixel);
+
+
+void FVYUConvertPS::SetParameters(FRHIBatchedShaderParameters& BatchedParameters, TRefCountPtr<FRHITexture> VYUTexture, const FIntPoint& OutputDimensions, const FMatrix44f& ColorTransform, UE::Color::EEncoding Encoding, const FMatrix44f& CSTransform)
+{
+	FVYUConvertUB UB;
+	{
+		UB.ColorTransform = ColorTransform;
+		UB.CSTransform = CSTransform;
+		UB.EOTF = int(Encoding);
+		UB.UVScale = FVector2f((float)OutputDimensions.X / (float)VYUTexture->GetSizeX(), (float)OutputDimensions.Y / (float)VYUTexture->GetSizeY());
+		UB.Texture = VYUTexture;
+		UB.Sampler = TStaticSamplerState<SF_Bilinear>::GetRHI();
+	}
+
+	TUniformBufferRef<FVYUConvertUB> Data = TUniformBufferRef<FVYUConvertUB>::CreateUniformBufferImmediate(UB, UniformBuffer_SingleFrame);
+	SetUniformBufferParameter(BatchedParameters, GetUniformBufferParameter<FVYUConvertUB>(), Data);
+}
diff --git a/Engine/Source/Runtime/RenderCore/Public/MediaShaders.h b/Engine/Source/Runtime/RenderCore/Public/MediaShaders.h
--- a/Engine/Source/Runtime/RenderCore/Public/MediaShaders.h
+++ b/Engine/Source/Runtime/RenderCore/Public/MediaShaders.h
@@ -640,9 +640,13 @@ public:
 
 	FReadTextureExternalPS(const ShaderMetaType::CompiledShaderInitializerType& Initializer)
 		: FGlobalShader(Initializer)
-	{ }
+	{
+		InTexture.Bind(Initializer.ParameterMap, TEXT("InExternalTexture"), SPF_Mandatory);
+	}
 
 	RENDERCORE_API void SetParameters(FRHIBatchedShaderParameters& BatchedParameters, FTextureRHIRef TextureExt, FSamplerStateRHIRef SamplerState, const FLinearColor & ScaleRotation, const FLinearColor & Offset);
+
+	LAYOUT_FIELD(FShaderResourceParameter, InTexture);
 };
 
 
@@ -723,5 +727,25 @@ public:
 	RENDERCORE_API FModifyAlphaSwizzleRgbaPS::FParameters* AllocateAndSetParameters(FRDGBuilder& GraphBuilder, FRDGTextureRef RGBATexture, FRDGTextureRef OutputTexture);
 };
 
+/**
+ * Pixel shader to convert a VYU frame to RGBA.
+ *
+ * This shader expects a VYU frame packed into a single texture in PF_R8G8B8A8
+ * format with the following memory layout: [V0, Y0, U0, unused][V1, Y1, U1, unused]..
+ *
+ * @see https://registry.khronos.org/vulkan/specs/1.3/html/vkspec.html#VkSamplerYcbcrModelConversion
+ */
+class FVYUConvertPS
+	: public FGlobalShader
+{
+	DECLARE_EXPORTED_SHADER_TYPE(FVYUConvertPS, Global, RENDERCORE_API);
+
+public:
+	FVYUConvertPS() { }
 
+	FVYUConvertPS(const ShaderMetaType::CompiledShaderInitializerType& Initializer)
+		: FGlobalShader(Initializer)
+	{ }
 
+	RENDERCORE_API void SetParameters(FRHIBatchedShaderParameters& BatchedParameters, TRefCountPtr<FRHITexture> VYUTexture, const FIntPoint& OutputDimensions, const FMatrix44f& ColorTransform, UE::Color::EEncoding Encoding, const FMatrix44f& CSTransform);
+};
