Electra: Moving utilities from ElectraPlayerRuntime to ElectraBase

#rnx

diff --git a/Engine/Plugins/Media/ElectraPlayer/Source/ElectraPlayerRuntime/Private/Runtime/Core/MediaCriticalSection.h b/Engine/Plugins/Media/ElectraPlayer/Source/ElectraPlayerRuntime/Private/Runtime/Core/MediaCriticalSection.h
deleted file mode 100644
diff --git a/Engine/Plugins/Media/ElectraPlayer/Source/ElectraPlayerRuntime/Private/Runtime/Core/MediaEventSignal.h b/Engine/Plugins/Media/ElectraPlayer/Source/ElectraPlayerRuntime/Private/Runtime/Core/MediaEventSignal.h
deleted file mode 100644
diff --git a/Engine/Plugins/Media/ElectraPlayer/Source/ElectraPlayerRuntime/Private/Runtime/Core/MediaInterlocked.h b/Engine/Plugins/Media/ElectraPlayer/Source/ElectraPlayerRuntime/Private/Runtime/Core/MediaInterlocked.h
deleted file mode 100644
diff --git a/Engine/Plugins/Media/ElectraPlayer/Source/ElectraPlayerRuntime/Private/Runtime/Core/MediaLock.h b/Engine/Plugins/Media/ElectraPlayer/Source/ElectraPlayerRuntime/Private/Runtime/Core/MediaLock.h
deleted file mode 100644
diff --git a/Engine/Plugins/Media/ElectraPlayer/Source/ElectraPlayerRuntime/Private/Runtime/Core/MediaMacros.h b/Engine/Plugins/Media/ElectraPlayer/Source/ElectraPlayerRuntime/Private/Runtime/Core/MediaMacros.h
deleted file mode 100644
diff --git a/Engine/Plugins/Media/ElectraPlayer/Source/ElectraPlayerRuntime/Private/Runtime/Core/MediaMessageQueue.h b/Engine/Plugins/Media/ElectraPlayer/Source/ElectraPlayerRuntime/Private/Runtime/Core/MediaMessageQueue.h
deleted file mode 100644
diff --git a/Engine/Plugins/Media/ElectraPlayer/Source/ElectraPlayerRuntime/Private/Runtime/Core/MediaNoncopyable.h b/Engine/Plugins/Media/ElectraPlayer/Source/ElectraPlayerRuntime/Private/Runtime/Core/MediaNoncopyable.h
deleted file mode 100644
diff --git a/Engine/Plugins/Media/ElectraPlayer/Source/ElectraPlayerRuntime/Private/Runtime/Core/MediaQueue.h b/Engine/Plugins/Media/ElectraPlayer/Source/ElectraPlayerRuntime/Private/Runtime/Core/MediaQueue.h
deleted file mode 100644
diff --git a/Engine/Plugins/Media/ElectraPlayer/Source/ElectraPlayerRuntime/Private/Runtime/Core/MediaSemaphore.h b/Engine/Plugins/Media/ElectraPlayer/Source/ElectraPlayerRuntime/Private/Runtime/Core/MediaSemaphore.h
deleted file mode 100644
diff --git a/Engine/Plugins/Media/ElectraPlayer/Source/ElectraPlayerRuntime/Private/Runtime/Core/MediaThreads.cpp b/Engine/Plugins/Media/ElectraPlayer/Source/ElectraPlayerRuntime/Private/Runtime/Core/MediaThreads.cpp
deleted file mode 100644
diff --git a/Engine/Plugins/Media/ElectraPlayer/Source/ElectraPlayerRuntime/Private/Runtime/Core/MediaThreads.h b/Engine/Plugins/Media/ElectraPlayer/Source/ElectraPlayerRuntime/Private/Runtime/Core/MediaThreads.h
deleted file mode 100644
diff --git a/Engine/Plugins/Media/ElectraPlayer/Source/ElectraPlayerRuntime/Private/Runtime/Core/MediaTypes.h b/Engine/Plugins/Media/ElectraPlayer/Source/ElectraPlayerRuntime/Private/Runtime/Core/MediaTypes.h
deleted file mode 100644
diff --git a/Engine/Plugins/Media/ElectraUtil/Source/ElectraBase/Private/Core/MediaThreads.cpp b/Engine/Plugins/Media/ElectraUtil/Source/ElectraBase/Private/Core/MediaThreads.cpp
new file mode 100644
--- /dev/null
+++ b/Engine/Plugins/Media/ElectraUtil/Source/ElectraBase/Private/Core/MediaThreads.cpp
@@ -0,0 +1,418 @@
+// Copyright Epic Games, Inc. All Rights Reserved.
+
+#include "Core/MediaThreads.h"
+#include "Core/MediaMessageQueue.h"
+
+#include "HAL/RunnableThread.h"
+#include "Misc/ScopeLock.h"
+
+
+namespace MediaRunnablePrivate
+{
+
+static TMediaMessageQueueDynamicNoTimeout<TFunction<void()>> AsyncFunctionsToRun;
+static volatile bool bEndForShutdown = false;
+static FMediaThread* AsyncExecutionThread = nullptr;
+static void AsyncWorkerThreadFN()
+{
+	while(!bEndForShutdown)
+	{
+		TFunction<void()> AsyncRunFN = AsyncFunctionsToRun.ReceiveMessage();
+		if (AsyncRunFN)
+		{
+			AsyncRunFN();
+		}
+	}
+}
+
+}
+
+void FMediaRunnable::Startup()
+{
+	MediaRunnablePrivate::bEndForShutdown = false;
+	MediaRunnablePrivate::AsyncExecutionThread = new FMediaThread("Electra::ExecAsync");
+	MediaRunnablePrivate::AsyncExecutionThread->ThreadStart(FMediaRunnable::FStartDelegate::CreateStatic(&MediaRunnablePrivate::AsyncWorkerThreadFN));
+}
+
+void FMediaRunnable::Shutdown()
+{
+	check(MediaRunnablePrivate::AsyncExecutionThread);
+	
+	TFunction<void()> FinishTask = []()
+	{
+		MediaRunnablePrivate::bEndForShutdown = true;
+	};
+	EnqueueAsyncTask(MoveTemp(FinishTask));
+
+	MediaRunnablePrivate::AsyncExecutionThread->ThreadWaitDone();
+	delete MediaRunnablePrivate::AsyncExecutionThread;
+	MediaRunnablePrivate::AsyncExecutionThread = nullptr;
+}
+
+void FMediaRunnable::EnqueueAsyncTask(TFunction<void()>&& InFunctionToExecuteOnAsyncThread)
+{
+	MediaRunnablePrivate::AsyncFunctionsToRun.SendMessage(MoveTemp(InFunctionToExecuteOnAsyncThread));
+}
+
+
+// ----------------------------------------------------------------------------
+/**
+ * Creates a new thread.
+ * Actually only the parameters are set, but no actual thread is created.
+ * We leave this to the Start() function so it becomes possible to change parameters after construction.
+ *
+ * @param InCoreAffinityMask
+ * @param InPriority
+ * @param InStackSize
+ * @param InThreadName
+ *
+ * @return New runnable object
+ */
+FMediaRunnable* FMediaRunnable::Create(int32 InCoreAffinityMask, EThreadPriority InPriority, uint32 InStackSize, const FString& InThreadName)
+{
+	if (InCoreAffinityMask == 0 || InCoreAffinityMask == -1)
+	{
+		InCoreAffinityMask = FPlatformAffinity::GetNoAffinityMask();
+	}
+
+	FMediaRunnable* t = new FMediaRunnable;
+	if (t)
+	{
+		t->ThreadPriority = InPriority;
+		t->InitialCoreAffinity = InCoreAffinityMask;
+		t->StackSize = Align((uint32)InStackSize, 16384);
+		t->ThreadName = InThreadName;
+		t->bIsStarted = false;
+	}
+	return t;
+}
+
+//-----------------------------------------------------------------------------
+/**
+ *
+ */
+void FMediaRunnable::Destroy(FMediaRunnable* InMediaRunnable)
+{
+	delete InMediaRunnable;
+}
+
+// ----------------------------------------------------------------------------
+/**
+ * Creates and starts the thread with the current settings.
+ *
+ * @param Entry
+ * @param bWaitRunning
+ */
+void FMediaRunnable::Start(FMediaRunnable::FStartDelegate Entry, bool bWaitRunning)
+{
+	// Remember entry function
+	EntryFunction = Entry;
+
+	StartInternal();
+
+	// Optionally wait for thread startup
+	if (bWaitRunning)
+	{
+		SignalRunning->Wait();
+	}
+}
+
+
+
+// ----------------------------------------------------------------------------
+/**
+ * Sets a 'done' signal to signal when the thread is being destroyed.
+ *
+ * @param pDoneSignal
+ */
+void FMediaRunnable::SetDoneSignal(FMediaEvent* InDoneSignal)
+{
+	FScopeLock Lock(&StateAccessMutex);
+	DoneSignal = InDoneSignal;
+}
+
+
+// ----------------------------------------------------------------------------
+/**
+ * Creates and starts the thread with the current settings.
+ */
+void FMediaRunnable::StartInternal()
+{
+	// We ignore stack size, priority and affinity here on purpose. Prio and affinity are all set to defaults and stack sizes are
+	// often set too small so we rather go with defaults.
+	MediaThreadRunnable = FRunnableThread::Create(this, *ThreadName);
+	checkf(MediaThreadRunnable, TEXT("Could not create RunnableThread!"));
+	if (MediaThreadRunnable)
+	{
+		bIsStarted = true;
+	}
+}
+
+
+
+// ----------------------------------------------------------------------------
+/**
+ * Thread startup function.
+ */
+uint32 FMediaRunnable::Run()
+{
+	SignalRunning->Trigger();
+	EntryFunction.Execute();
+	return 0;
+}
+
+
+// ----------------------------------------------------------------------------
+/**
+ * Thread completion function.
+ * Only notifies optionally set signal which FMediaThread::ThreadWaitDone() is waiting for.
+ * Actual thread deletion will happen from there.
+ */
+void FMediaRunnable::Exit()
+{
+	StateAccessMutex.Lock();
+	bIsStarted = false;
+	FMediaEvent* Sig = DoneSignal;
+	DoneSignal = nullptr;
+	StateAccessMutex.Unlock();
+	if (Sig)
+	{
+		Sig->Signal();
+	}
+}
+
+
+// ----------------------------------------------------------------------------
+/**
+ * Thread constructor
+ */
+FMediaRunnable::FMediaRunnable()
+{
+	MediaThreadRunnable = nullptr;
+
+	ThreadPriority = TPri_Normal;
+	StackSize = 0;
+
+	DoneSignal = nullptr;
+	bIsStarted = false;
+
+	InitialCoreAffinity = FPlatformAffinity::GetNoAffinityMask();
+	SignalRunning = FPlatformProcess::GetSynchEventFromPool();
+}
+
+
+
+// ----------------------------------------------------------------------------
+/**
+ * Thread destructor.
+ */
+FMediaRunnable::~FMediaRunnable()
+{
+	delete MediaThreadRunnable;
+	FPlatformProcess::ReturnSynchEventToPool(SignalRunning);
+}
+
+
+
+
+
+// ----------------------------------------------------------------------------
+/**
+ * Changes the priority of the thread.
+ *
+ * @param newPriority
+ *               New thread priority
+ *
+ * @return Previous priority
+ */
+EThreadPriority FMediaRunnable::ChangePriority(EThreadPriority newPriority)
+{
+	checkf(!bIsStarted, TEXT("FMediaRunnable::ChangePriority: Missing interface how to change thread priortity when thread was already started"));
+	EThreadPriority prevPrio = ThreadPriority;
+	ThreadPriority = newPriority;
+	return prevPrio;
+}
+
+
+// ----------------------------------------------------------------------------
+/**
+ * Give a new name to the thread.
+ *
+ * @param pNewName New name of the thread.
+ */
+void FMediaRunnable::SetName(const FString& InThreadName)
+{
+	checkf(!bIsStarted, TEXT("FMediaRunnable::Name: Cannot set name when thread is already running"));
+	ThreadName = InThreadName;
+}
+
+// ----------------------------------------------------------------------------
+/**
+ * Constructor
+ *
+ * @param pName
+ * @param osPriority
+ * @param stackSize
+ */
+FMediaThread::FMediaThread(const char* AnsiName)
+	: MediaRunnable(nullptr)
+	, Priority(TPri_Normal)
+	, StackSize(0)
+	, CoreAffinity(0)
+	, bIsStarted(false)
+	, bWaitDoneOnDelete(false)
+{
+	ThreadSetName(AnsiName);
+}
+
+
+
+// ----------------------------------------------------------------------------
+/**
+ * Destructor
+ */
+FMediaThread::~FMediaThread()
+{
+	if (bWaitDoneOnDelete && bIsStarted)
+	{
+		SigDone.Wait();
+	}
+	if (MediaRunnable)
+	{
+		FMediaRunnable::Destroy(MediaRunnable);
+	}
+}
+
+
+// ----------------------------------------------------------------------------
+/**
+ * Set a thread priority other than the one given to the constructor.
+ *
+ * @param osPriority
+ */
+void FMediaThread::ThreadSetPriority(EThreadPriority InPriority)
+{
+	Priority = InPriority;
+	if (MediaRunnable)
+	{
+		MediaRunnable->ChangePriority(Priority);
+	}
+}
+
+
+
+// ----------------------------------------------------------------------------
+/**
+ * Sets a thread core affinity mask before starting the thread.
+ *
+ * @param coreAffinity
+ */
+void FMediaThread::ThreadSetCoreAffinity(int32 coreAffinity)
+{
+	checkf(!bIsStarted, TEXT("not while the thread is already running!"));
+	CoreAffinity = coreAffinity;
+}
+
+
+
+// ----------------------------------------------------------------------------
+/**
+ * Set a thread stack size other than the one given to the constructor before starting the thread.
+ *
+ * @param stackSize
+ */
+void FMediaThread::ThreadSetStackSize(uint32 stackSize)
+{
+	checkf(!bIsStarted, TEXT("not while the thread is already running!"));
+	StackSize = stackSize;
+}
+
+
+
+// ----------------------------------------------------------------------------
+/**
+ * Set a thread name other than the one given to the constructor before starting the thread.
+ *
+ * @param pName
+ */
+void FMediaThread::ThreadSetName(const char* InAnsiThreadName)
+{
+	checkf(!bIsStarted, TEXT("not while the thread is already running!"));
+	ThreadName = FString(InAnsiThreadName);
+}
+
+
+
+// ----------------------------------------------------------------------------
+/**
+ * Sets a new thread name once the thread is running.
+ *
+ * @param pName
+ */
+void FMediaThread::ThreadRename(const char* InAnsiThreadName)
+{
+	checkf(bIsStarted, TEXT("only while the thread is already running!"));
+	ThreadName = FString(InAnsiThreadName);
+	MediaRunnable->SetName(ThreadName);
+}
+
+
+
+// ----------------------------------------------------------------------------
+/**
+ * Sets whether or not the destructor needs to wait for the thread to have finished. Defaults to false. Useful when using this as a member variable.
+ *
+ * @param bWait
+ */
+void FMediaThread::ThreadWaitDoneOnDelete(bool bWait)
+{
+	bWaitDoneOnDelete = bWait;
+}
+
+
+
+// ----------------------------------------------------------------------------
+/**
+ * Waits for the thread to have finished.
+ */
+void FMediaThread::ThreadWaitDone(void)
+{
+	checkf(bIsStarted, TEXT("thread not started"));
+	SigDone.Wait();
+}
+
+
+
+// ----------------------------------------------------------------------------
+/**
+ * Thread start function.
+ *
+ * @param entryFunction
+ */
+void FMediaThread::ThreadStart(FMediaRunnable::FStartDelegate EntryFunction)
+{
+	MediaRunnable = FMediaRunnable::Create(CoreAffinity, Priority, StackSize, ThreadName);
+	MediaRunnable->SetDoneSignal(&SigDone);
+	bIsStarted = true;
+	MediaRunnable->Start(EntryFunction, false);
+}
+
+
+
+// ----------------------------------------------------------------------------
+/**
+ * Resets the thread to be started again. Must have waited for thread termination using ThreadWaitDone() first!
+ */
+void FMediaThread::ThreadReset()
+{
+	check(!bIsStarted || (bIsStarted && SigDone.IsSignaled()));
+	if (MediaRunnable)
+	{
+		FMediaRunnable::Destroy(MediaRunnable);
+	}
+	MediaRunnable = nullptr;
+	bIsStarted = false;
+	SigDone.Reset();
+}
+
+
diff --git a/Engine/Plugins/Media/ElectraUtil/Source/ElectraBase/Public/Core/MediaCriticalSection.h b/Engine/Plugins/Media/ElectraUtil/Source/ElectraBase/Public/Core/MediaCriticalSection.h
new file mode 100644
--- /dev/null
+++ b/Engine/Plugins/Media/ElectraUtil/Source/ElectraBase/Public/Core/MediaCriticalSection.h
@@ -0,0 +1,59 @@
+// Copyright Epic Games, Inc. All Rights Reserved.
+
+#pragma once
+
+#include "Misc/Build.h"
+#include "HAL/CriticalSection.h"
+#include "Core/MediaTypes.h"
+#include "Core/MediaNoncopyable.h"
+
+
+class FMediaCriticalSection : private TMediaNoncopyable<FMediaCriticalSection>
+{
+public:
+	FMediaCriticalSection() = default;
+
+	~FMediaCriticalSection()
+	{
+		CriticalSection.Lock();
+		CriticalSection.Unlock();
+	}
+
+	void Lock() const
+	{
+		CriticalSection.Lock();
+	}
+
+	void Unlock() const
+	{
+		CriticalSection.Unlock();
+	}
+
+	bool TryLock() const
+	{
+		return CriticalSection.TryLock();
+	}
+
+	/**
+	 * Scoped lock helper class.
+	*/
+	class ScopedLock : private TMediaNoncopyable<ScopedLock>
+	{
+	public:
+		explicit ScopedLock(const FMediaCriticalSection& lock)
+			: CriticalSection(lock)
+		{
+			CriticalSection.Lock();
+		}
+		~ScopedLock()
+		{
+			CriticalSection.Unlock();
+		}
+	private:
+		ScopedLock();
+		const FMediaCriticalSection& CriticalSection;
+	};
+
+private:
+	mutable FCriticalSection CriticalSection;
+};
diff --git a/Engine/Plugins/Media/ElectraUtil/Source/ElectraBase/Public/Core/MediaEventSignal.h b/Engine/Plugins/Media/ElectraUtil/Source/ElectraBase/Public/Core/MediaEventSignal.h
new file mode 100644
--- /dev/null
+++ b/Engine/Plugins/Media/ElectraUtil/Source/ElectraBase/Public/Core/MediaEventSignal.h
@@ -0,0 +1,90 @@
+// Copyright Epic Games, Inc. All Rights Reserved.
+
+#pragma once
+
+#include "Misc/Build.h"
+#include "HAL/Event.h"
+#include "HAL/PlatformProcess.h"
+#include "Core/MediaNoncopyable.h"
+
+
+class FMediaEvent : private TMediaNoncopyable<FMediaEvent>
+{
+public:
+	FMediaEvent() : Event(FPlatformProcess::GetSynchEventFromPool(true))
+	{ }
+
+	~FMediaEvent()
+	{
+		Event->Trigger();
+		FPlatformProcess::ReturnSynchEventToPool(Event);
+	}
+
+	/**
+	 * Waits for the event signal to get set. If already set, this function returns immediately.
+	 */
+	void Wait()
+	{
+		Event->Wait();
+	}
+
+	/**
+	 * Waits for the event signal to get set and clears it immediately. If already set, this function returns immediately.
+	 */
+	void WaitAndReset()
+	{
+		Event->Wait();
+		Event->Reset();
+	}
+
+	/**
+	 * Signals the event to be set. If already set, this does nothing.
+	 */
+	void Signal()
+	{
+		Event->Trigger();
+	}
+
+	/**
+	 * Clears the signal. Threads will now wait again until the signal is set once more.
+	 */
+	void Reset()
+	{
+		Event->Reset();
+	}
+
+	/**
+	 * Returns state of the signal.
+	 */
+	bool IsSignaled() const
+	{
+		return Event->Wait(0);
+	}
+
+	/**
+	 * Waits for the event signal to get set within the given time limit.
+	 * If already set, this function returns immediately. Returns true when event was set, false when wait timed out.
+	 */
+	bool WaitTimeout(int64 MicroSeconds)
+	{
+		return Event->Wait(FTimespan::FromMicroseconds(MicroSeconds));
+	}
+
+	/**
+	 * Waits for the event signal to get set within the given time limit and clears it immediately.
+	 * If already set, this function returns immediately. Returns true when event was set, false when wait timed out.
+	 */
+	bool WaitTimeoutAndReset(int64 MicroSeconds)
+	{
+		check(MicroSeconds > 0);	// 0 is forbidden. either call Wait() without timeout, or check the signal using IsSignaled()
+		bool bGot = Event->Wait(FTimespan::FromMicroseconds(MicroSeconds));
+		if (bGot)
+		{
+			Event->Reset();
+		}
+		return bGot;
+	}
+
+private:
+	mutable FEvent* Event;
+};
diff --git a/Engine/Plugins/Media/ElectraUtil/Source/ElectraBase/Public/Core/MediaInterlocked.h b/Engine/Plugins/Media/ElectraUtil/Source/ElectraBase/Public/Core/MediaInterlocked.h
new file mode 100644
--- /dev/null
+++ b/Engine/Plugins/Media/ElectraUtil/Source/ElectraBase/Public/Core/MediaInterlocked.h
@@ -0,0 +1,102 @@
+// Copyright Epic Games, Inc. All Rights Reserved.
+
+#pragma once
+
+#include "Core/MediaTypes.h"
+#include "Core/MediaMacros.h"
+#include "CoreGlobals.h"
+#include "HAL/PlatformAtomics.h"
+
+inline uint32 FMediaInterlockedRead(uint32 volatile& variable)
+{
+	return (uint32)FPlatformAtomics::AtomicRead_Relaxed((volatile const int32*)&variable);
+}
+
+inline uint32 FMediaInterlockedIncrement(uint32 volatile& variable)
+{
+	// The UE implementation of FPlatformAtomics returnes the new value. But media assumes that the value before the operaton is returned
+	return (uint32)(FPlatformAtomics::InterlockedIncrement((volatile int32*)&variable) - 1);
+}
+
+inline uint32 FMediaInterlockedDecrement(uint32 volatile& variable)
+{
+	// The UE implementation of FPlatformAtomics returnes the new value. But media assumes that the value before the operaton is returned
+	return (uint32)(FPlatformAtomics::InterlockedDecrement((volatile int32*)&variable) + 1);
+}
+
+inline uint32 FMediaInterlockedAdd(uint32 volatile& variable, uint32 value)
+{
+	return (uint32)FPlatformAtomics::InterlockedAdd((volatile int32*)&variable, (int32)value);
+}
+
+inline uint32 FMediaInterlockedExchange(uint32 volatile& variable, uint32 exchangeValue)
+{
+	return (uint32)FPlatformAtomics::InterlockedExchange((volatile int32*)&variable, (int32)exchangeValue);
+}
+
+inline void* FMediaInterlockedExchangePointerVoid(void* volatile& variable, void* pExchangeValue)
+{
+	return FPlatformAtomics::InterlockedExchangePtr((void**)&variable, pExchangeValue);
+}
+
+inline uint32 FMediaInterlockedCompareExchange(uint32 volatile& variable, uint32 exchangeValue, uint32 compareValue)
+{
+	return (uint32)FPlatformAtomics::InterlockedCompareExchange((volatile int32*)&variable, (int32)exchangeValue, (int32)compareValue);
+}
+
+inline void* FMediaInterlockedCompareExchangePointer(void* volatile& variable, void* pExchangeValue, void* pCompareValue)
+{
+#if PLATFORM_64BITS
+	return (void*)FPlatformAtomics::InterlockedCompareExchange((volatile int64*)&variable, (int64)pExchangeValue, (int64)pCompareValue);
+#else
+	return (void*)FPlatformAtomics::InterlockedCompareExchange((volatile int32*)&variable, (int32)pExchangeValue, (int32)pCompareValue);
+#endif
+}
+
+inline uint64 FMediaInterlockedRead64(uint64 volatile& variable)
+{
+	return (uint64)FPlatformAtomics::AtomicRead_Relaxed((volatile const int64*)&variable);
+}
+
+inline uint64 FMediaInterlockedAdd64(uint64 volatile& variable, uint64 value)
+{
+	return (uint64)FPlatformAtomics::InterlockedAdd((volatile int64*)&variable, (int64)value);
+}
+
+inline int32 FMediaInterlockedRead(int32 volatile& variable)
+{
+	return FPlatformAtomics::AtomicRead_Relaxed(&variable);
+}
+
+inline int32 FMediaInterlockedIncrement(int32 volatile& variable)
+{
+	// The UE implementation of FPlatformAtomics returnes the new value. But media assumes that the value before the operaton is returned
+	return FPlatformAtomics::InterlockedIncrement(&variable) - 1;
+}
+
+inline int32 FMediaInterlockedDecrement(int32 volatile& variable)
+{
+	// The UE implementation of FPlatformAtomics returnes the new value. But media assumes that the value before the operaton is returned
+	return FPlatformAtomics::InterlockedDecrement(&variable) + 1;
+}
+
+inline int32 FMediaInterlockedAdd(int32 volatile& variable, int32 value)
+{
+	return FPlatformAtomics::InterlockedAdd(&variable, value);
+}
+
+inline int32 FMediaInterlockedExchange(int32 volatile& variable, int32 exchangeValue)
+{
+	return FPlatformAtomics::InterlockedExchange(&variable, exchangeValue);
+}
+
+inline int32 FMediaInterlockedCompareExchange(int32 volatile& variable, int32 exchangeValue, int32 compareValue)
+{
+	return FPlatformAtomics::InterlockedCompareExchange(&variable, exchangeValue, compareValue);
+}
+
+template <typename T, typename X>
+inline T* TMediaInterlockedExchangePointer(T* volatile& variable, X* pExchangeValue)
+{
+	return((T*)FMediaInterlockedExchangePointerVoid((void* volatile&)variable, (void*)pExchangeValue));
+}
diff --git a/Engine/Plugins/Media/ElectraUtil/Source/ElectraBase/Public/Core/MediaLock.h b/Engine/Plugins/Media/ElectraUtil/Source/ElectraBase/Public/Core/MediaLock.h
new file mode 100644
--- /dev/null
+++ b/Engine/Plugins/Media/ElectraUtil/Source/ElectraBase/Public/Core/MediaLock.h
@@ -0,0 +1,70 @@
+// Copyright Epic Games, Inc. All Rights Reserved.
+
+#pragma once
+
+#include "Misc/Build.h"
+#include "HAL/CriticalSection.h"
+#include "Core/MediaNoncopyable.h"
+
+/**
+ * Empty lock class.
+ */
+class FMediaLockNone : private TMediaNoncopyable<FMediaLockNone>
+{
+public:
+	inline void Lock() const {}
+	void Unlock() const {}
+
+	class ScopedLock : private TMediaNoncopyable<ScopedLock>
+	{
+	public:
+		explicit ScopedLock(const FMediaLockNone& lock) : LockToUse(lock)
+		{ }
+		~ScopedLock() = default;
+	private:
+		ScopedLock();
+		const FMediaLockNone& LockToUse;
+	};
+};
+
+
+
+/**
+ * Lock class for generic container classes requiring mutual exclusion access.
+ */
+class FMediaLockCriticalSection : private TMediaNoncopyable<FMediaLockCriticalSection>
+{
+public:
+	FMediaLockCriticalSection()
+	{
+		CriticalSection.Lock();
+		CriticalSection.Unlock();
+	}
+
+	void Lock() const
+	{
+		CriticalSection.Lock();
+	}
+	void Unlock() const
+	{
+		CriticalSection.Unlock();
+	}
+
+	class ScopedLock : private TMediaNoncopyable<ScopedLock>
+	{
+	public:
+		explicit ScopedLock(const FMediaLockCriticalSection& lock) : LockToUse(lock)
+		{
+			LockToUse.Lock();
+		}
+		~ScopedLock()
+		{
+			LockToUse.Unlock();
+		}
+	private:
+		ScopedLock();
+		const FMediaLockCriticalSection& LockToUse;
+	};
+private:
+	mutable FCriticalSection CriticalSection;
+};
diff --git a/Engine/Plugins/Media/ElectraUtil/Source/ElectraBase/Public/Core/MediaMacros.h b/Engine/Plugins/Media/ElectraUtil/Source/ElectraBase/Public/Core/MediaMacros.h
new file mode 100644
--- /dev/null
+++ b/Engine/Plugins/Media/ElectraUtil/Source/ElectraBase/Public/Core/MediaMacros.h
@@ -0,0 +1,121 @@
+// Copyright Epic Games, Inc. All Rights Reserved.
+
+#pragma once
+
+#include "Core/MediaTypes.h"
+
+//--------------------------------------------------------------------------
+//
+// Array convenience macros
+//
+// Evaluates to the number of elements in a statically defined array.
+#define FMEDIA_STATIC_ARRAY_COUNT(array)			(SIZE_T)(sizeof((array)) / sizeof((array)[0]))
+
+//--------------------------------------------------------------------------
+//
+// Endian swapping
+//
+static inline uint16 MEDIA_ENDIAN_SWAP(uint16 value)
+{
+	return((value >> 8) | (value << 8));
+}
+
+static inline int16 MEDIA_ENDIAN_SWAP(int16 value)
+{
+	return(int16(MEDIA_ENDIAN_SWAP(uint16(value))));
+}
+
+static inline uint32 MEDIA_ENDIAN_SWAP(uint32 value)
+{
+	return((value << 24) | ((value & 0xff00) << 8) | ((value >> 8) & 0xff00) | (value >> 24));
+}
+
+static inline int32 MEDIA_ENDIAN_SWAP(int32 value)
+{
+	return(int32(MEDIA_ENDIAN_SWAP(uint32(value))));
+}
+
+static inline uint64 MEDIA_ENDIAN_SWAP(uint64 value)
+{
+	return((uint64(MEDIA_ENDIAN_SWAP(uint32(value & 0xffffffffU))) << 32) | uint64(MEDIA_ENDIAN_SWAP(uint32(value >> 32))));
+}
+
+static inline int64 MEDIA_ENDIAN_SWAP(int64 value)
+{
+	return(int64(MEDIA_ENDIAN_SWAP(uint64(value))));
+}
+
+
+// Yes, there's no endian swapping with 8 bit values. We put them here so templated functions can be used with these types.
+static inline uint8  MEDIA_TO_BIG_ENDIAN(uint8 value)			{ return(value); }
+static inline int8   MEDIA_TO_BIG_ENDIAN(int8 value)			{ return(value); }
+static inline uint8  MEDIA_TO_LITTLE_ENDIAN(uint8 value)		{ return(value); }
+static inline int8   MEDIA_TO_LITTLE_ENDIAN(int8 value)			{ return(value); }
+static inline uint8  MEDIA_FROM_BIG_ENDIAN(uint8 value)			{ return(value); }
+static inline int8   MEDIA_FROM_BIG_ENDIAN(int8 value)			{ return(value); }
+static inline uint8  MEDIA_FROM_LITTLE_ENDIAN(uint8 value)		{ return(value); }
+static inline int8   MEDIA_FROM_LITTLE_ENDIAN(int8 value)		{ return(value); }
+
+
+#if !PLATFORM_LITTLE_ENDIAN
+// Convert a native value to big endian
+static inline uint16 MEDIA_TO_BIG_ENDIAN(uint16 value)		{ return(value); }
+static inline int16 MEDIA_TO_BIG_ENDIAN(int16 value)		{ return(value); }
+static inline int32 MEDIA_TO_BIG_ENDIAN(int32 value)		{ return(value); }
+static inline uint32 MEDIA_TO_BIG_ENDIAN(uint32 value)		{ return(value); }
+static inline int64 MEDIA_TO_BIG_ENDIAN(int64 value)		{ return(value); }
+static inline uint64 MEDIA_TO_BIG_ENDIAN(uint64 value)		{ return(value); }
+// Convert a native value to little endian
+static inline uint16 MEDIA_TO_LITTLE_ENDIAN(uint16 value)	{ return(MEDIA_ENDIAN_SWAP(value)); }
+static inline int16 MEDIA_TO_LITTLE_ENDIAN(int16 value)		{ return(MEDIA_ENDIAN_SWAP(value)); }
+static inline int32 MEDIA_TO_LITTLE_ENDIAN(int32 value)		{ return(MEDIA_ENDIAN_SWAP(value)); }
+static inline uint32 MEDIA_TO_LITTLE_ENDIAN(uint32 value)	{ return(MEDIA_ENDIAN_SWAP(value)); }
+static inline int64 MEDIA_TO_LITTLE_ENDIAN(int64 value)		{ return(MEDIA_ENDIAN_SWAP(value)); }
+static inline uint64 MEDIA_TO_LITTLE_ENDIAN(uint64 value)	{ return(MEDIA_ENDIAN_SWAP(value)); }
+// Convert a big endian value to native
+static inline uint16 MEDIA_FROM_BIG_ENDIAN(uint16 value)	{ return(value); }
+static inline int16 MEDIA_FROM_BIG_ENDIAN(int16 value)		{ return(value); }
+static inline int32 MEDIA_FROM_BIG_ENDIAN(int32 value)		{ return(value); }
+static inline uint32 MEDIA_FROM_BIG_ENDIAN(uint32 value)	{ return(value); }
+static inline int64 MEDIA_FROM_BIG_ENDIAN(int64 value)		{ return(value); }
+static inline uint64 MEDIA_FROM_BIG_ENDIAN(uint64 value)	{ return(value); }
+// Convert a little endian value to native
+static inline uint16 MEDIA_FROM_LITTLE_ENDIAN(uint16 value)	{ return(MEDIA_ENDIAN_SWAP(value)); }
+static inline int16 MEDIA_FROM_LITTLE_ENDIAN(int16 value)	{ return(MEDIA_ENDIAN_SWAP(value)); }
+static inline int32 MEDIA_FROM_LITTLE_ENDIAN(int32 value)	{ return(MEDIA_ENDIAN_SWAP(value)); }
+static inline uint32 MEDIA_FROM_LITTLE_ENDIAN(uint32 value)	{ return(MEDIA_ENDIAN_SWAP(value)); }
+static inline int64 MEDIA_FROM_LITTLE_ENDIAN(int64 value)	{ return(MEDIA_ENDIAN_SWAP(value)); }
+static inline uint64 MEDIA_FROM_LITTLE_ENDIAN(uint64 value)	{ return(MEDIA_ENDIAN_SWAP(value)); }
+
+#else
+
+// Convert a native value to big endian
+static inline uint16 MEDIA_TO_BIG_ENDIAN(uint16 value)		{ return(MEDIA_ENDIAN_SWAP(value)); }
+static inline int16 MEDIA_TO_BIG_ENDIAN(int16 value)		{ return(MEDIA_ENDIAN_SWAP(value)); }
+static inline int32 MEDIA_TO_BIG_ENDIAN(int32 value)		{ return(MEDIA_ENDIAN_SWAP(value)); }
+static inline uint32 MEDIA_TO_BIG_ENDIAN(uint32 value)		{ return(MEDIA_ENDIAN_SWAP(value)); }
+static inline int64 MEDIA_TO_BIG_ENDIAN(int64 value)		{ return(MEDIA_ENDIAN_SWAP(value)); }
+static inline uint64 MEDIA_TO_BIG_ENDIAN(uint64 value)		{ return(MEDIA_ENDIAN_SWAP(value)); }
+// Convert a native value to little endian
+static inline uint16 MEDIA_TO_LITTLE_ENDIAN(uint16 value)	{ return(value); }
+static inline int16 MEDIA_TO_LITTLE_ENDIAN(int16 value)		{ return(value); }
+static inline int32 MEDIA_TO_LITTLE_ENDIAN(int32 value)		{ return(value); }
+static inline uint32 MEDIA_TO_LITTLE_ENDIAN(uint32 value)	{ return(value); }
+static inline int64 MEDIA_TO_LITTLE_ENDIAN(int64 value)		{ return(value); }
+static inline uint64 MEDIA_TO_LITTLE_ENDIAN(uint64 value)	{ return(value); }
+// Convert a big endian value to native
+static inline uint16 MEDIA_FROM_BIG_ENDIAN(uint16 value)	{ return(MEDIA_ENDIAN_SWAP(value)); }
+static inline int16 MEDIA_FROM_BIG_ENDIAN(int16 value)		{ return(MEDIA_ENDIAN_SWAP(value)); }
+static inline int32 MEDIA_FROM_BIG_ENDIAN(int32 value)		{ return(MEDIA_ENDIAN_SWAP(value)); }
+static inline uint32 MEDIA_FROM_BIG_ENDIAN(uint32 value)	{ return(MEDIA_ENDIAN_SWAP(value)); }
+static inline int64 MEDIA_FROM_BIG_ENDIAN(int64 value)		{ return(MEDIA_ENDIAN_SWAP(value)); }
+static inline uint64 MEDIA_FROM_BIG_ENDIAN(uint64 value)	{ return(MEDIA_ENDIAN_SWAP(value)); }
+// Convert a little endian value to native
+static inline uint16 MEDIA_FROM_LITTLE_ENDIAN(uint16 value)	{ return(value); }
+static inline int16 MEDIA_FROM_LITTLE_ENDIAN(int16 value)	{ return(value); }
+static inline int32 MEDIA_FROM_LITTLE_ENDIAN(int32 value)	{ return(value); }
+static inline uint32 MEDIA_FROM_LITTLE_ENDIAN(uint32 value)	{ return(value); }
+static inline int64 MEDIA_FROM_LITTLE_ENDIAN(int64 value)	{ return(value); }
+static inline uint64 MEDIA_FROM_LITTLE_ENDIAN(uint64 value)	{ return(value); }
+
+#endif
diff --git a/Engine/Plugins/Media/ElectraUtil/Source/ElectraBase/Public/Core/MediaMessageQueue.h b/Engine/Plugins/Media/ElectraUtil/Source/ElectraBase/Public/Core/MediaMessageQueue.h
new file mode 100644
--- /dev/null
+++ b/Engine/Plugins/Media/ElectraUtil/Source/ElectraBase/Public/Core/MediaMessageQueue.h
@@ -0,0 +1,665 @@
+// Copyright Epic Games, Inc. All Rights Reserved.
+
+#pragma once
+
+#include "Core/MediaTypes.h"
+#include "Core/MediaMacros.h"
+#include "Core/MediaNoncopyable.h"
+#include "Core/MediaSemaphore.h"
+#include "Core/MediaQueue.h"
+
+
+/**
+ * Generic message queue WITHOUT timeout support.
+ *
+ * Conceptually a message queue is a multiple-sender, single-receiver object, so use it only in that way!
+ *
+ * This version supports only a maximum number of messages that needs to be specified in the constructor
+ * or by calling Resize() before first use.
+ */
+template <typename T>
+class TMediaMessageQueueNoTimeout
+{
+public:
+	// ----------------------------------------------------------------------------
+	/**
+	 * Constructor.
+	 */
+	TMediaMessageQueueNoTimeout(SIZE_T MaxMessages = 0)
+		: Messages(MaxMessages)
+		, AvailSema(MaxMessages)
+	{
+	}
+
+	// ----------------------------------------------------------------------------
+	/**
+	 * Waits indefinitely for a message to arrive.
+	 *
+	 * @return Received message
+	 */
+	T ReceiveMessage()
+	{
+		ReadySema.Obtain();
+		T Msg = Messages.Pop();
+		AvailSema.Release();
+		return Msg;
+	}
+
+	// ----------------------------------------------------------------------------
+	/**
+	 * Receives a message if one is already pending in the queue.
+	 *
+	 * @param message Receives the message if true is returned, unchanged otherwise
+	 *
+	 * @return true if message received and stored in 'message', false if no message received.
+	 */
+	bool ReceiveMessage(T& Message)
+	{
+		bool bHave;
+		bHave = ReadySema.TryToObtain();
+		if (bHave)
+		{
+			Message = Messages.Pop();
+			AvailSema.Release();
+			return true;
+		}
+		return false;
+	}
+
+	// ----------------------------------------------------------------------------
+	/**
+	 * Sends a message.
+	 *
+	 * @param message Message to send
+	 * @param bWaitIfFull true to wait for the queue to have room if it's full, false to return immediately.
+	 *
+	 * @return true if message sent successfully, false if queue is full and bWaitIfFull was false.
+	 */
+	bool SendMessage(const T& Message, bool bWaitIfFull = true)
+	{
+		bool bOk = bWaitIfFull ? AvailSema.Obtain() : AvailSema.TryToObtain();
+		if (bOk)
+		{
+			Messages.Push(Message);
+			ReadySema.Release();
+			return true;
+		}
+		return false;
+	}
+
+	// ----------------------------------------------------------------------------
+	/**
+	 * Jams a message to the front of the queue.
+	 *
+	 * @param message Message to send
+	 * @param bWaitIfFull true to wait for the queue to have room if it's full, false to return immediately.
+	 *
+	 * @return true if message sent successfully, false if queue is full and bWaitIfFull was false.
+	 */
+	bool JamMessage(const T& Message, bool bWaitIfFull = true)
+	{
+		bool bOk = bWaitIfFull ? AvailSema.Obtain() : AvailSema.TryToObtain();
+		if (bOk)
+		{
+			Messages.PushFront(Message);
+			ReadySema.Release();
+			return true;
+		}
+		return false;
+	}
+
+	// ----------------------------------------------------------------------------
+	/**
+	 * Checks if there are pending messages.
+	 *
+	 * @return true if there are pending messages, false if not.
+	 */
+	bool HaveMessage() const
+	{
+		return !Messages.IsEmpty();
+	}
+
+	// ----------------------------------------------------------------------------
+	/**
+	 * Queries the number of messages in the queue.
+	 *
+	 * @return Number of enqueued messages.
+	 */
+	SIZE_T NumWaitingMessages() const
+	{
+		return Messages.Num();
+	}
+
+	// ----------------------------------------------------------------------------
+	/**
+	 * Resizes the message queue, discarding any messages it may hold.
+	 *
+	 * Should only be used after the queue was created, but before it gets used!
+	 *
+	 * @param maxMessages New message queue size.
+	 *
+	 * @return none
+	 */
+	void Resize(SIZE_T MaxMessages)
+	{
+		ReadySema.SetCount(0);
+		AvailSema.SetCount(MaxMessages);
+		Messages.Resize(MaxMessages);
+	}
+
+private:
+	TMediaQueue<T, FMediaLockCriticalSection> Messages;
+	FMediaSemaphore AvailSema;
+	FMediaSemaphore ReadySema;
+};
+
+
+
+
+/**
+ * Generic message queue WITH timeout support.
+ *
+ * Conceptually a message queue is a multiple-sender, single-receiver object, so use it only in that way!
+ *
+ * This version supports only a maximum number of messages that needs to be specified in the constructor
+ * or by calling Resize() before first use.
+ */
+template <typename T>
+class TMediaMessageQueueWithTimeout
+{
+public:
+	// ----------------------------------------------------------------------------
+	/**
+	 * Constructor.
+	 */
+	TMediaMessageQueueWithTimeout(SIZE_T MaxMessages = 0)
+		: Messages(MaxMessages)
+		, AvailSema(MaxMessages)
+	{
+	}
+
+	// ----------------------------------------------------------------------------
+	/**
+	 * Waits indefinitely for a message to arrive.
+	 *
+	 * @return Received message
+	 */
+	T ReceiveMessage()
+	{
+		ReadySema.Obtain();
+		T Msg = Messages.Pop();
+		AvailSema.Release();
+		return Msg;
+	}
+
+	// ----------------------------------------------------------------------------
+	/**
+	 * Receives a message if one is already pending in the queue.
+	 *
+	 * @param message Receives the message if true is returned, unchanged otherwise
+	 *
+	 * @return true if message received and stored in 'message', false if no message received.
+	 */
+	bool ReceiveMessage(T& Message)
+	{
+		bool bHave;
+		bHave = ReadySema.TryToObtain();
+		if (bHave)
+		{
+			Message = Messages.Pop();
+			AvailSema.Release();
+			return true;
+		}
+		return false;
+	}
+
+	// ----------------------------------------------------------------------------
+	/**
+	 * Waits a maximum amount of time for a message to arrive.
+	 *
+	 * @param message Receives the message if true is returned, unchanged otherwise
+	 * @param maxWaitMicroseconds Maximum wait time in microseconds
+	 *
+	 * @return true if message received and stored in 'message', false if no message received.
+	 */
+	bool ReceiveMessage(T& Message, int64 MaxWaitMicroseconds)
+	{
+		check(MaxWaitMicroseconds >= 0);
+		bool bHave;
+		bHave = ReadySema.Obtain(MaxWaitMicroseconds);
+		if (bHave)
+		{
+			Message = Messages.Pop();
+			AvailSema.Release();
+			return true;
+		}
+		return false;
+	}
+
+	// ----------------------------------------------------------------------------
+	/**
+	 * Sends a message.
+	 *
+	 * @param message Message to send
+	 * @param bWaitIfFull true to wait for the queue to have room if it's full, false to return immediately.
+	 *
+	 * @return true if message sent successfully, false if queue is full and bWaitIfFull was false.
+	 */
+	bool SendMessage(const T& Message, bool bWaitIfFull = true)
+	{
+		bool bOk = bWaitIfFull ? AvailSema.Obtain() : AvailSema.TryToObtain();
+		if (bOk)
+		{
+			Messages.Push(Message);
+			ReadySema.Release();
+			return true;
+		}
+		return false;
+	}
+
+	// ----------------------------------------------------------------------------
+	/**
+	 * Jams a message to the front of the queue.
+	 *
+	 * @param message Message to send
+	 * @param bWaitIfFull true to wait for the queue to have room if it's full, false to return immediately.
+	 *
+	 * @return true if message sent successfully, false if queue is full and bWaitIfFull was false.
+	 */
+	bool JamMessage(const T& Message, bool bWaitIfFull = true)
+	{
+		bool bOk = bWaitIfFull ? AvailSema.Obtain() : AvailSema.TryToObtain();
+		if (bOk)
+		{
+			Messages.PushFront(Message);
+			ReadySema.Release();
+			return true;
+		}
+		return false;
+	}
+
+	// ----------------------------------------------------------------------------
+	/**
+	 * Checks if there are pending messages.
+	 *
+	 * @return true if there are pending messages, false if not.
+	 */
+	bool HaveMessage() const
+	{
+		return !Messages.IsEmpty();
+	}
+
+	// ----------------------------------------------------------------------------
+	/**
+	 * Queries the number of messages in the queue.
+	 *
+	 * @return Number of enqueued messages.
+	 */
+	SIZE_T NumWaitingMessages() const
+	{
+		return Messages.Num();
+	}
+
+	// ----------------------------------------------------------------------------
+	/**
+	 * Resizes the message queue, discarding any messages it may hold.
+	 *
+	 * Should only be used after the queue was created, but before it gets used!
+	 *
+	 * @param maxMessages New message queue size.
+	 *
+	 * @return none
+	 */
+	void Resize(SIZE_T MaxMessages)
+	{
+		ReadySema.SetCount(0);
+		AvailSema.SetCount(MaxMessages);
+		Messages.Resize(MaxMessages);
+	}
+
+private:
+	TMediaQueue<T, FMediaLockCriticalSection> Messages;
+	FMediaSemaphore AvailSema;
+	FMediaSemaphore ReadySema;
+};
+
+
+
+
+
+/**
+ * Generic message queue WITHOUT timeout support.
+ *
+ * Conceptually a message queue is a multiple-sender, single-receiver object, so use it only in that way!
+ *
+ * This version uses a template argument to define the initial size.
+**/
+template <typename T, SIZE_T CAPACITY>
+class TMediaMessageQueueStaticNoTimeout : public TMediaMessageQueueNoTimeout<T>
+{
+public:
+	// ----------------------------------------------------------------------------
+	/**
+	 * Constructor.
+	**/
+	TMediaMessageQueueStaticNoTimeout()
+		: TMediaMessageQueueNoTimeout<T>(CAPACITY)
+	{
+	}
+};
+
+
+/**
+ * Generic message queue WITH timeout support.
+ *
+ * Conceptually a message queue is a multiple-sender, single-receiver object, so use it only in that way!
+ *
+ * This version uses a template argument to define the initial size.
+**/
+template <typename T, SIZE_T CAPACITY>
+class TMediaMessageQueueStaticWithTimeout : public TMediaMessageQueueWithTimeout<T>
+{
+public:
+	// ----------------------------------------------------------------------------
+	/**
+	 * Constructor.
+	**/
+	TMediaMessageQueueStaticWithTimeout()
+		: TMediaMessageQueueWithTimeout<T>(CAPACITY)
+	{
+	}
+};
+
+
+
+
+
+
+
+
+
+
+/**
+ * Generic message queue WITHOUT timeout support.
+ *
+ * Conceptually a message queue is a multiple-sender, single-receiver object, so use it only in that way!
+ *
+ * This version supports an unlimited number of messages (we limit it to 2^31-1)
+ * Sending a message can therefore never block. We keep the interface identical to TMediaMessageQueueNoTimeout::SendMessage
+ * for compatibility.
+**/
+template <typename T>
+class TMediaMessageQueueDynamicNoTimeout
+{
+public:
+	// ----------------------------------------------------------------------------
+	/**
+	 * Constructor.
+	**/
+	TMediaMessageQueueDynamicNoTimeout()
+		: AvailSema(0x7fffffff)
+	{
+	}
+
+	// ----------------------------------------------------------------------------
+	/**
+	 * Waits indefinitely for a message to arrive.
+	 *
+	 * @return		Received message
+	**/
+	T ReceiveMessage()
+	{
+		ReadySema.Obtain();
+		T Msg = Messages.Pop();
+		AvailSema.Release();
+		return Msg;
+	}
+
+	// ----------------------------------------------------------------------------
+	/**
+	 * Receives a message if one is already pending in the queue.
+	 *
+	 * @param		message     Receives the message if true is returned, unchanged otherwise
+	 *
+	 * @return		true if message received and stored in 'message', false if no message received.
+	**/
+	bool ReceiveMessage(T& Message)
+	{
+		bool bHave;
+		bHave = ReadySema.TryToObtain();
+		if (bHave)
+		{
+			Message = Messages.Pop();
+			AvailSema.Release();
+			return true;
+		}
+		return false;
+	}
+
+	// ----------------------------------------------------------------------------
+	/**
+	 * Sends a message.
+	 *
+	 * @param		message         Message to send
+	 * @param		bWaitIfFull     true to wait for the queue to have room if it's full, false to return immediately.
+	 *
+	 * @return		true if message sent successfully, false if queue is full and bWaitIfFull was false.
+	**/
+	bool SendMessage(const T& Message, bool bWaitIfFull = true)
+	{
+		bool bOk = bWaitIfFull ? AvailSema.Obtain() : AvailSema.TryToObtain();
+		if (bOk)
+		{
+			Messages.Push(Message);
+			ReadySema.Release();
+			return true;
+		}
+		return false;
+	}
+
+	// ----------------------------------------------------------------------------
+	/**
+	 * Jams a message to the front of the queue.
+	 *
+	 * @param		message         Message to send
+	 * @param		bWaitIfFull     true to wait for the queue to have room if it's full, false to return immediately.
+	 *
+	 * @return		true if message sent successfully, false if queue is full and bWaitIfFull was false.
+	**/
+	bool JamMessage(const T& Message, bool bWaitIfFull = true)
+	{
+		bool bOk = bWaitIfFull ? AvailSema.Obtain() : AvailSema.TryToObtain();
+		if (bOk)
+		{
+			Messages.PushFront(Message);
+			ReadySema.Release();
+			return true;
+		}
+		return false;
+	}
+
+	// ----------------------------------------------------------------------------
+	/**
+	 * Checks if there are pending messages.
+	 *
+	 * @return		true if there are pending messages, false if not.
+	**/
+	bool HaveMessage() const
+	{
+		return !Messages.IsEmpty();
+	}
+
+	// ----------------------------------------------------------------------------
+	/**
+	 * Queries the number of messages in the queue.
+	 *
+	 * @return		Number of enqueued messages.
+	**/
+	SIZE_T NumWaitingMessages() const
+	{
+		return Messages.Num();
+	}
+
+private:
+	TMediaQueueDynamic<T>	Messages;
+	FMediaSemaphore			AvailSema;
+	FMediaSemaphore			ReadySema;
+};
+
+
+
+
+
+
+/**
+ * Generic message queue WITH timeout support.
+ *
+ * Conceptually a message queue is a multiple-sender, single-receiver object, so use it only in that way!
+ *
+ * This version supports an unlimited number of messages (we limit it to 2^31-1)
+ * Sending a message can therefore never block. We keep the interface identical to TMediaMessageQueueNoTimeout::SendMessage
+ * for compatibility.
+**/
+template <typename T>
+class TMediaMessageQueueDynamicWithTimeout
+{
+public:
+	// ----------------------------------------------------------------------------
+	/**
+	 * Constructor.
+	**/
+	TMediaMessageQueueDynamicWithTimeout()
+		: AvailSema(0x7fffffff)
+	{
+	}
+
+	// ----------------------------------------------------------------------------
+	/**
+	 * Waits indefinitely for a message to arrive.
+	 *
+	 * @return		Received message
+	**/
+	T ReceiveMessage()
+	{
+		ReadySema.Obtain();
+		T Msg = Messages.Pop();
+		AvailSema.Release();
+		return Msg;
+	}
+
+	// ----------------------------------------------------------------------------
+	/**
+	 * Receives a message if one is already pending in the queue.
+	 *
+	 * @param		message     Receives the message if true is returned, unchanged otherwise
+	 *
+	 * @return		true if message received and stored in 'message', false if no message received.
+	**/
+	bool ReceiveMessage(T& Message)
+	{
+		bool bHave;
+		bHave = ReadySema.TryToObtain();
+		if (bHave)
+		{
+			Message = Messages.Pop();
+			AvailSema.Release();
+			return true;
+		}
+		return false;
+	}
+
+	// ----------------------------------------------------------------------------
+	/**
+	 * Waits a maximum amount of time for a message to arrive.
+	 *
+	 * @param		message                 Receives the message if true is returned, unchanged otherwise
+	 * @param		maxWaitMicroseconds     Maximum wait time in microseconds
+	 *
+	 * @return		true if message received and stored in 'message', false if no message received.
+	**/
+	bool ReceiveMessage(T& Message, int64 MaxWaitMicroseconds)
+	{
+		check(MaxWaitMicroseconds >= 0);
+		bool bHave;
+		bHave = ReadySema.Obtain(MaxWaitMicroseconds);
+		if (bHave)
+		{
+			Message = Messages.Pop();
+			AvailSema.Release();
+			return true;
+		}
+		return false;
+	}
+
+	// ----------------------------------------------------------------------------
+	/**
+	 * Sends a message.
+	 *
+	 * @param		message         Message to send
+	 * @param		bWaitIfFull     true to wait for the queue to have room if it's full, false to return immediately.
+	 *
+	 * @return		true if message sent successfully, false if queue is full and bWaitIfFull was false.
+	**/
+	bool SendMessage(const T& Message, bool bWaitIfFull = true)
+	{
+		bool bOk = bWaitIfFull ? AvailSema.Obtain() : AvailSema.TryToObtain();
+		if (bOk)
+		{
+			Messages.Push(Message);
+			ReadySema.Release();
+			return true;
+		}
+		return false;
+	}
+
+	// ----------------------------------------------------------------------------
+	/**
+	 * Jams a message to the front of the queue.
+	 *
+	 * @param		message         Message to send
+	 * @param		bWaitIfFull     true to wait for the queue to have room if it's full, false to return immediately.
+	 *
+	 * @return		true if message sent successfully, false if queue is full and bWaitIfFull was false.
+	**/
+	bool JamMessage(const T& Message, bool bWaitIfFull = true)
+	{
+		bool bOk = bWaitIfFull ? AvailSema.Obtain() : AvailSema.TryToObtain();
+		if (bOk)
+		{
+			Messages.PushFront(Message);
+			ReadySema.Release();
+			return true;
+		}
+		return false;
+	}
+
+	// ----------------------------------------------------------------------------
+	/**
+	 * Checks if there are pending messages.
+	 *
+	 * @return		true if there are pending messages, false if not.
+	**/
+	bool HaveMessage() const
+	{
+		return !Messages.IsEmpty();
+	}
+
+	// ----------------------------------------------------------------------------
+	/**
+	 * Queries the number of messages in the queue.
+	 *
+	 * @return		Number of enqueued messages.
+	**/
+	SIZE_T NumWaitingMessages() const
+	{
+		return Messages.Num();
+	}
+
+private:
+	TMediaQueueDynamic<T>			Messages;
+	FMediaSemaphore					AvailSema;
+	FMediaSemaphore					ReadySema;
+};
+
+
+
diff --git a/Engine/Plugins/Media/ElectraUtil/Source/ElectraBase/Public/Core/MediaNoncopyable.h b/Engine/Plugins/Media/ElectraUtil/Source/ElectraBase/Public/Core/MediaNoncopyable.h
new file mode 100644
--- /dev/null
+++ b/Engine/Plugins/Media/ElectraUtil/Source/ElectraBase/Public/Core/MediaNoncopyable.h
@@ -0,0 +1,14 @@
+// Copyright Epic Games, Inc. All Rights Reserved.
+
+#pragma once
+
+template<typename T>
+class TMediaNoncopyable
+{
+protected:
+	TMediaNoncopyable() = default;
+	~TMediaNoncopyable() = default;
+private:
+	TMediaNoncopyable(const TMediaNoncopyable&) = delete;
+	TMediaNoncopyable& operator = (const TMediaNoncopyable&) = delete;
+};
diff --git a/Engine/Plugins/Media/ElectraUtil/Source/ElectraBase/Public/Core/MediaQueue.h b/Engine/Plugins/Media/ElectraUtil/Source/ElectraBase/Public/Core/MediaQueue.h
new file mode 100644
--- /dev/null
+++ b/Engine/Plugins/Media/ElectraUtil/Source/ElectraBase/Public/Core/MediaQueue.h
@@ -0,0 +1,852 @@
+// Copyright Epic Games, Inc. All Rights Reserved.
+
+#pragma once
+
+#include "Core/MediaTypes.h"
+#include "Core/MediaMacros.h"
+#include "Core/MediaLock.h"
+#include "Core/MediaInterlocked.h"
+#include "Core/MediaNoncopyable.h"
+
+
+/**
+ * Default memory allocator.
+ */
+template <typename T>
+class TMediaQueueAllocator
+{
+public:
+	static T* Allocate(uint32 NumElements)
+	{
+		return reinterpret_cast<T*>(FMemory::Malloc(NumElements * sizeof(T),  __alignof(T)));
+	}
+	static void Deallocate(T* Address)
+	{
+		if (Address)
+		{
+			FMemory::Free((void*)Address);
+		}
+	}
+};
+
+
+
+/**
+ * Object safe queue class.
+ *
+ * Implements by-value methods only to avoid working with references.
+ *
+ * Methods like
+ *  T& Push()
+ *  T& operator[](int32)
+ *  T& Front()
+ *
+ * are dangerous when used in a multi-threaded environment since the elements cannot be protected by an internal mutex.
+ * If you are using an external mutex to ensure object safety you may add such functions, but please derive a new class for this!!!
+ *
+ */
+template <typename T, typename L = FMediaLockNone, typename M = TMediaQueueAllocator<T> >
+class TMediaQueue : public L
+{
+public:
+	using MyType = TMediaQueue<T, L, M>;
+	using ElementType = T;
+
+	TMediaQueue(int32 InitialCapacity = 0)
+	{
+		Resize(InitialCapacity);
+	}
+
+	// Not copyable (for now) because of the possible mutex class.
+	TMediaQueue(const MyType& rhs) = delete;
+	MyType& operator = (const MyType& rhs) = delete;
+
+	~TMediaQueue()
+	{
+		Resize(0);
+	}
+
+	void Clear()
+	{
+		L::Lock();
+		for(int32 i=0; i<NumIn; ++i)
+		{
+			Elements[IdxOut].~T();
+			if (++IdxOut == MaxNum)
+			{
+				IdxOut -= MaxNum;
+			}
+		}
+		IdxIn = IdxOut = NumIn = 0;
+		L::Unlock();
+	}
+
+	void Resize(int32 newCapacity)
+	{
+		L::Lock();
+		Clear();
+		if (newCapacity != MaxNum)
+		{
+			M::Deallocate(Elements);
+			if (newCapacity)
+			{
+				Elements = M::Allocate(newCapacity);
+			}
+			else
+			{
+				Elements = nullptr;
+			}
+			MaxNum = newCapacity;
+		}
+		L::Unlock();
+	}
+
+	void Swap(MyType& rhs)
+	{
+		L::Lock();
+		// Remember our values.
+		T* pElements = Elements;
+		int32  nCapacity = MaxNum;
+		int32  nNumIn = NumIn;
+		int32  nIdxIn = IdxIn;
+		int32  nIdxOut = IdxOut;
+		// Get the values from the other queue.
+		Elements = rhs.Elements;
+		MaxNum = rhs.MaxNum;
+		NumIn = rhs.NumIn;
+		IdxIn = rhs.IdxIn;
+		IdxOut = rhs.IdxOut;
+		// Set our old values in the other queue.
+		rhs.Elements = pElements;
+		rhs.MaxNum = nCapacity;
+		rhs.NumIn = nNumIn;
+		rhs.IdxIn = nIdxIn;
+		rhs.IdxOut = nIdxOut;
+		L::Unlock();
+	}
+
+	int32 Capacity() const
+	{
+		return MaxNum;
+	}
+
+	int32 Num() const
+	{
+		return NumIn;
+	}
+
+	bool IsEmpty() const
+	{
+		return Num() == 0;
+	}
+
+	bool IsFull() const
+	{
+		return Num() == Capacity();
+	}
+
+	void Push(const T& Element)
+	{
+		L::Lock();
+		check(!IsFull());
+		int32 pos = IdxIn;
+		if (++IdxIn == MaxNum)
+		{
+			IdxIn = 0;
+		}
+		++NumIn;
+		// Copying the element can NOT take place outside the lock in case multiple threads push at the same time and a Resize() will happen.
+		new ((void*)(Elements + pos)) T(Element);
+		L::Unlock();
+	}
+
+	void PushFront(const T& Element)
+	{
+		L::Lock();
+		check(!IsFull());
+		if (IdxOut == 0)
+		{
+			IdxOut = MaxNum - 1;
+		}
+		else
+		{
+			--IdxOut;
+		}
+		int32 pos = IdxOut;
+		++NumIn;
+		// Copying the element can NOT take place outside the lock in case multiple threads push at the same time and a Resize() will happen.
+		new ((void*)(Elements + pos)) T(Element);
+		L::Unlock();
+	}
+
+	T Pop()
+	{
+		L::Lock();
+		check(!IsEmpty());
+		T r(Elements[IdxOut]);
+		Elements[IdxOut].~T();
+		if (++IdxOut == MaxNum)
+		{
+			IdxOut = 0;
+		}
+		--NumIn;
+		L::Unlock();
+		return r;
+	}
+
+	T Front() const
+	{
+		L::Lock();
+		check(!IsEmpty());
+		T r(Elements[IdxOut]);
+		L::Unlock();
+		return r;
+	}
+
+protected:
+	T* Elements = nullptr;
+	int32 MaxNum = 0;
+	int32 NumIn = 0;
+	int32 IdxIn = 0;
+	int32 IdxOut = 0;
+};
+
+
+
+
+
+/**
+ * A mutex-less version of TMediaQueue to enable working with references.
+ *
+ * Since this is not thread-safe you must ensure access with an external mutex.
+ */
+template <typename T, typename M = TMediaQueueAllocator<T> >
+class TMediaQueueNoLock : public TMediaQueue<T, FMediaLockNone, M>
+{
+public:
+	using MyType = TMediaQueueNoLock<T, M>;
+	using ElementType = T;
+	typedef TMediaQueue<T, FMediaLockNone, M>	B;
+
+	TMediaQueueNoLock() : B() {}
+	~TMediaQueueNoLock() {}
+
+	T& PushRef()
+	{
+		check(!B::IsFull());
+		int32 pos = B::IdxIn;
+		if (++B::IdxIn == B::Capacity())
+		{
+			B::IdxIn = 0;
+		}
+		++B::NumIn;
+		new ((void*)(B::Elements + pos)) T();
+		return B::Elements[pos];
+	}
+
+	const T& FrontRef() const
+	{
+		check(!B::IsEmpty());
+		return B::Elements[B::IdxOut];
+	}
+
+	T& FrontRef()
+	{
+		check(!B::IsEmpty());
+		return B::Elements[B::IdxOut];
+	}
+
+	const T& BackRef() const
+	{
+		check(!B::IsEmpty());
+		return B::Elements[(B::IdxIn + B::Capacity() - 1) % B::Capacity()];
+	}
+
+	T& BackRef()
+	{
+		check(!B::IsEmpty());
+		return B::Elements[(B::IdxIn + B::Capacity() - 1) % B::Capacity()];
+	}
+
+	void Erase(int32 i)
+	{
+		check(!B::IsEmpty());
+		check(i < B::NumIn);
+		--B::NumIn;
+		int32 i0 = (i + B::IdxOut) % B::Capacity();
+		int32 i1 = (i + B::IdxOut + 1) % B::Capacity();
+		for(; i < B::NumIn; ++i)
+		{
+			B::Elements[i0] = B::Elements[i1];
+			i0 = i1++;
+			if (i1 >= B::Capacity())
+			{
+				i1 -= B::Capacity();
+			}
+		}
+		B::Elements[i0].~T();
+		if (B::IdxIn == 0)
+		{
+			B::IdxIn = B::Capacity() - 1;
+		}
+		else
+		{
+			--B::IdxIn;
+		}
+	}
+
+	const T& operator[](int32 i) const
+	{
+		check(!B::IsEmpty());
+		check(i < B::NumIn);
+		i += B::IdxOut;
+		if (i >= B::Capacity())
+		{
+			i -= B::Capacity();
+		}
+		return B::Elements[i];
+	}
+
+	T& operator[](int32 i)
+	{
+		check(!B::IsEmpty());
+		check(i < B::NumIn);
+		i += B::IdxOut;
+		if (i >= B::Capacity())
+		{
+			i -= B::Capacity();
+		}
+		return B::Elements[i];
+	}
+
+	void AppendFirstElements(const MyType& FromOther, int32 MaxElements = ~int32(0))
+	{
+		// NOTE: Your code needs to make sure the source queue is protected by some external lock so does not get
+		//       modified while we are appending its elements.
+		int32 NumAvail = B::Capacity() - B::Num();
+		int32 NumToAppend = FromOther.Num();
+		if (NumToAppend > NumAvail)
+		{
+			NumToAppend = NumAvail;
+		}
+		if (NumToAppend > MaxElements)
+		{
+			NumToAppend = MaxElements;
+		}
+		for(int32 i = 0; i < NumToAppend; ++i)
+		{
+			B::Push(FromOther[i]);
+		}
+	}
+
+	void AppendLastElements(const MyType& FromOther, int32 MaxElements = ~int32(0))
+	{
+		// NOTE: Your code needs to make sure the source queue is protected by some external lock so does not get
+		//       modified while we're appending its elements.
+		int32 NumAvail = B::Capacity() - B::Num();
+		int32 NumToAppend = FromOther.Num();
+		if (NumToAppend > NumAvail)
+		{
+			NumToAppend = NumAvail;
+		}
+		if (NumToAppend > MaxElements)
+		{
+			NumToAppend = MaxElements;
+		}
+		for(int32 i = 0, j = FromOther.Num() - NumToAppend; i < NumToAppend; ++i, ++j)
+		{
+			B::Push(FromOther[j]);
+		}
+	}
+
+};
+
+
+
+
+
+/**
+ * Object safe queue class.
+ *
+ * Implements by-value methods only to avoid working with references.
+ *
+ * Methods like
+ *  T& Push()
+ *  T& operator[](int32)
+ *  T& Front()
+ *
+ * are dangerous when used in a multi-threaded environment since the elements cannot be protected by an internal mutex.
+ * If you are using an external mutex to ensure object safety you may add such functions, but please derive a new class for this!!!
+ * (or use the one below)
+ *
+ */
+template <typename T, uint32 CAPACITY, typename L = FMediaLockCriticalSection >
+class TMediaQueueFixedStatic : public L
+{
+public:
+	using MyType = TMediaQueueFixedStatic<T, CAPACITY, L>;
+	using ElementType = T;
+
+	TMediaQueueFixedStatic() = default;
+
+	// Not copyable (for now) because of the possible mutex class.
+	TMediaQueueFixedStatic(const MyType& rhs) = delete;
+	MyType& operator = (const MyType& rhs) = delete;
+
+	~TMediaQueueFixedStatic()
+	{
+		Clear();
+	}
+
+	void Clear()
+	{
+		L::Lock();
+		for(int32 i=0; i<NumIn; ++i)
+		{
+			Elements[IdxOut].~T();
+			if (++IdxOut == CAPACITY)
+			{
+				IdxOut -= CAPACITY;
+			}
+		}
+		IdxIn = IdxOut = NumIn = 0;
+		L::Unlock();
+	}
+
+	int32 Capacity() const
+	{
+		return CAPACITY;
+	}
+
+	int32 Num() const
+	{
+		return NumIn;
+	}
+
+	bool IsEmpty() const
+	{
+		return Num() == 0;
+	}
+
+	bool IsFull() const
+	{
+		return Num() == Capacity();
+	}
+
+	void Push(const T& element)
+	{
+		L::Lock();
+		check(!IsFull());
+		int32 pos = IdxIn;
+		if (++IdxIn == CAPACITY)
+		{
+			IdxIn = 0;
+		}
+		++NumIn;
+		// Copying the element can NOT take place outside the lock in case multiple threads push at the same time and a Resize() will happen.
+		new ((void*)(Elements + pos)) T(element);
+		L::Unlock();
+	}
+
+	void PushFront(const T& element)
+	{
+		L::Lock();
+		check(!IsFull());
+		if (IdxOut == 0)
+		{
+			IdxOut = CAPACITY - 1;
+		}
+		else
+		{
+			--IdxOut;
+		}
+		int32 pos = IdxOut;
+		++NumIn;
+		// Copying the element can NOT take place outside the lock in case multiple threads push at the same time and a Resize() will happen.
+		new ((void*)(Elements + pos)) T(element);
+		L::Unlock();
+	}
+
+	T Pop()
+	{
+		L::Lock();
+		check(!IsEmpty());
+		T r(Elements[IdxOut]);
+		Elements[IdxOut].~T();
+		if (++IdxOut == CAPACITY)
+		{
+			IdxOut = 0;
+		}
+		--NumIn;
+		L::Unlock();
+		return r;
+	}
+
+	T Front() const
+	{
+		L::Lock();
+		check(!IsEmpty());
+		T r(Elements[IdxOut]);
+		L::Unlock();
+		return r;
+	}
+
+protected:
+
+	MS_ALIGN(16)
+	uint8 FixedElements[CAPACITY * sizeof(T)]
+	GCC_ALIGN(16);
+	T* Elements = (T*)FixedElements;
+	int32 NumIn = 0;
+	int32 IdxIn = 0;
+	int32 IdxOut = 0;
+};
+
+
+
+template <typename T, uint32 CAPACITY>
+class TMediaQueueFixedStaticNoLock : public TMediaQueueFixedStatic<T, CAPACITY, FMediaLockNone>
+{
+public:
+	using MyType = TMediaQueueFixedStaticNoLock<T, CAPACITY>;
+	using ElementType = T;
+	typedef TMediaQueueFixedStatic<T, CAPACITY, FMediaLockNone>	B;
+
+	TMediaQueueFixedStaticNoLock() : B() {}
+	~TMediaQueueFixedStaticNoLock() {}
+
+	T& PushRef()
+	{
+		check(!B::IsFull());
+		int32 pos = B::IdxIn;
+		if (++B::IdxIn == CAPACITY)
+		{
+			B::IdxIn = 0;
+		}
+		++B::NumIn;
+		new ((void*)(B::Elements + pos)) T();
+		return B::Elements[pos];
+	}
+
+	const T& FrontRef() const
+	{
+		check(!B::IsEmpty());
+		return B::Elements[B::IdxOut];
+	}
+
+	T& FrontRef()
+	{
+		check(!B::IsEmpty());
+		return B::Elements[B::IdxOut];
+	}
+
+	const T& BackRef() const
+	{
+		check(!B::IsEmpty());
+		return B::Elements[(B::IdxIn + B::Capacity() - 1) % B::Capacity()];
+	}
+
+	T& BackRef()
+	{
+		check(!B::IsEmpty());
+		return B::Elements[(B::IdxIn + B::Capacity() - 1) % B::Capacity()];
+	}
+
+	void Erase(int32 i)
+	{
+		check(!B::IsEmpty());
+		check(i < B::NumIn);
+		--B::NumIn;
+		int32 i0 = (i + B::IdxOut) % CAPACITY;
+		int32 i1 = (i + B::IdxOut + 1) % CAPACITY;
+		for(; i < B::NumIn; ++i)
+		{
+			B::Elements[i0] = B::Elements[i1];
+			i0 = i1++;
+			if (i1 >= CAPACITY)
+			{
+				i1 -= CAPACITY;
+			}
+		}
+		B::Elements[i0].~T();
+		if (B::IdxIn == 0)
+		{
+			B::IdxIn = CAPACITY - 1;
+		}
+		else
+		{
+			--B::IdxIn;
+		}
+	}
+
+	const T& operator[](int32 i) const
+	{
+		check(!B::IsEmpty());
+		check(i < B::NumIn);
+		i += B::IdxOut;
+		if (i >= CAPACITY)
+		{
+			i -= CAPACITY;
+		}
+		return B::Elements[i];
+	}
+
+	T& operator[](int32 i)
+	{
+		check(!B::IsEmpty());
+		check(i < B::NumIn);
+		i += B::IdxOut;
+		if (i >= CAPACITY)
+		{
+			i -= CAPACITY;
+		}
+		return B::Elements[i];
+	}
+};
+
+
+
+
+
+/**
+ * Dynamically sized queue.
+ *
+ * Will grow when capacity is reached.
+ * Therefore a lock is REQUIRED! Never use this with a null-mutex unless you know what you're doing.
+ */
+template <typename T, typename L = FMediaLockCriticalSection, typename M = TMediaQueueAllocator<T> >
+class TMediaQueueDynamic : public L, private TMediaNoncopyable<TMediaQueueDynamic<T> >
+{
+public:
+	using MyType = TMediaQueueDynamic<T, L, M>;
+	using ElementType = T;
+
+	TMediaQueueDynamic(int32 initialCapacity = 0, int32 increment = 32)
+		: IncrementBy(increment)
+	{
+		Resize(initialCapacity);
+	}
+
+	~TMediaQueueDynamic()
+	{
+		L::Lock();
+		Resize(0);
+		L::Unlock();
+	}
+
+	void SetIncrement(int32 newIncrement)
+	{
+		L::Lock();
+		IncrementBy = newIncrement;
+		L::Unlock();
+	}
+
+	void Clear()
+	{
+		L::Lock();
+		for(int32 i=0; i<NumIn; ++i)
+		{
+			Elements[IdxOut].~T();
+			if (++IdxOut == MaxNum)
+			{
+				IdxOut -= MaxNum;
+			}
+		}
+		IdxIn = IdxOut = NumIn = 0;
+		L::Unlock();
+	}
+
+	void Resize(int32 newCapacity)
+	{
+		L::Lock();
+		if (newCapacity == 0)
+		{
+			Clear();
+			M::Deallocate(Elements);
+			Elements = nullptr;
+			MaxNum = 0;
+		}
+		else
+		{
+			// Create a new array and copy all active elements over.
+			checkf(newCapacity >= Num(), TEXT("Cannot shrink the dynamic queue to %u with %u elements currently managed!"), (uint32)newCapacity, (uint32)Num());
+			T* pNewElements = M::Allocate(newCapacity);
+			for(int32 i = 0, j = IdxOut; i < NumIn; ++i)
+			{
+				new ((void*)(pNewElements + i)) T(Elements[j]);
+				Elements[j].~T();
+				if (++j == MaxNum)
+				{
+					j = 0;
+				}
+			}
+			M::Deallocate(Elements);
+			Elements = pNewElements;
+			IdxIn = NumIn;
+			IdxOut = 0;
+			MaxNum = newCapacity;
+		}
+		L::Unlock();
+	}
+
+	int32 Increment() const
+	{
+		return IncrementBy;
+	}
+
+	int32 Capacity() const
+	{
+		return MaxNum;
+	}
+
+	int32 Num() const
+	{
+		return NumIn;
+	}
+
+	bool IsEmpty() const
+	{
+		return Num() == 0;
+	}
+
+	bool IsFull() const
+	{
+		return Num() == Capacity();
+	}
+
+	void Push(const T& element)
+	{
+		L::Lock();
+		if (IsFull())
+		{
+			Resize(Capacity() + Increment());
+		}
+		int32 pos = IdxIn;
+		if (++IdxIn == MaxNum)
+		{
+			IdxIn = 0;
+		}
+		++NumIn;
+		// Copying the element can NOT take place outside the lock in case multiple threads push at the same time and a Resize() will happen.
+		new ((void*)(Elements + pos)) T(element);
+		L::Unlock();
+	}
+
+	void PushFront(const T& element)
+	{
+		L::Lock();
+		if (IsFull())
+		{
+			Resize(Capacity() + Increment());
+		}
+		if (IdxOut == 0)
+		{
+			IdxOut = MaxNum - 1;
+		}
+		else
+		{
+			--IdxOut;
+		}
+		int32 pos = IdxOut;
+		++NumIn;
+		// Copying the element can NOT take place outside the lock in case multiple threads push at the same time and a Resize() will happen.
+		new ((void*)(Elements + pos)) T(element);
+		L::Unlock();
+	}
+
+	T Pop()
+	{
+		L::Lock();
+		check(!IsEmpty());
+		T r(Elements[IdxOut]);
+		Elements[IdxOut].~T();
+		if (++IdxOut == MaxNum)
+		{
+			IdxOut = 0;
+		}
+		--NumIn;
+		L::Unlock();
+		return r;
+	}
+
+	T Front() const
+	{
+		L::Lock();
+		check(!IsEmpty());
+		T r(Elements[IdxOut]);
+		L::Unlock();
+		return r;
+	}
+
+protected:
+	T* Elements = nullptr;
+	int32 MaxNum = 0;
+	int32 NumIn = 0;
+	int32 IdxIn = 0;
+	int32 IdxOut = 0;
+	int32 IncrementBy = 32;
+};
+
+
+
+
+
+template <typename T, typename M = TMediaQueueAllocator<T> >
+class TMediaQueueDynamicNoLock : public TMediaQueueDynamic<T, FMediaLockNone, M>
+{
+public:
+	using MyType = TMediaQueueDynamicNoLock<T, M>;
+	using ElementType = T;
+	typedef TMediaQueueDynamic<T, FMediaLockNone, M>	B;
+
+	TMediaQueueDynamicNoLock() : B() {}
+	~TMediaQueueDynamicNoLock() {}
+
+	const T& FrontRef() const
+	{
+		check(!B::IsEmpty());
+		return B::Elements[B::IdxOut];
+	}
+
+	T& FrontRef()
+	{
+		check(!B::IsEmpty());
+		return B::Elements[B::IdxOut];
+	}
+
+	const T& BackRef() const
+	{
+		check(!B::IsEmpty());
+		return B::Elements[(B::IdxIn + B::Capacity() - 1) % B::Capacity()];
+	}
+
+	T& BackRef()
+	{
+		check(!B::IsEmpty());
+		return B::Elements[(B::IdxIn + B::Capacity() - 1) % B::Capacity()];
+	}
+
+	const T& operator[](int32 i) const
+	{
+		check(!B::IsEmpty());
+		check(i < B::NumIn);
+		i += B::IdxOut;
+		if (i >= B::MaxNum)
+		{
+			i -= B::MaxNum;
+		}
+		return B::Elements[i];
+	}
+
+	T& operator[](int32 i)
+	{
+		check(!B::IsEmpty());
+		check(i < B::NumIn);
+		i += B::IdxOut;
+		if (i >= B::MaxNum)
+		{
+			i -= B::MaxNum;
+		}
+		return B::Elements[i];
+	}
+};
diff --git a/Engine/Plugins/Media/ElectraUtil/Source/ElectraBase/Public/Core/MediaSemaphore.h b/Engine/Plugins/Media/ElectraUtil/Source/ElectraBase/Public/Core/MediaSemaphore.h
new file mode 100644
--- /dev/null
+++ b/Engine/Plugins/Media/ElectraUtil/Source/ElectraBase/Public/Core/MediaSemaphore.h
@@ -0,0 +1,139 @@
+// Copyright Epic Games, Inc. All Rights Reserved.
+
+#pragma once
+
+#include "Misc/Build.h"
+#include "HAL/Event.h"
+#include "HAL/PlatformProcess.h"
+#include "Misc/ScopeLock.h"
+#include "Misc/DateTime.h"
+#include "Core/MediaNoncopyable.h"
+
+
+class FMediaSemaphore : private TMediaNoncopyable<FMediaSemaphore>
+{
+public:
+	FMediaSemaphore(int32 InitialCount = 0)
+		: Event(FPlatformProcess::GetSynchEventFromPool(true))
+		, Count(InitialCount)
+	{ }
+
+	~FMediaSemaphore()
+	{
+		CriticalSection.Lock();
+		Event->Trigger();
+		FPlatformProcess::ReturnSynchEventToPool(Event);
+		Event = nullptr;
+		CriticalSection.Unlock();
+		Count = 0;
+	}
+
+	bool Obtain()
+	{
+		CriticalSection.Lock();
+
+		// Check if decrement operation would work. We inspect the value before the decrment operation
+		// INFO: The UE implementation of InterlockedDecrement returns the value after the decrement operation
+		while(FPlatformAtomics::InterlockedDecrement(&Count) <= -1)
+		{
+			// Undo the "wrong" decrement operation to restore value...
+			FPlatformAtomics::InterlockedIncrement(&Count);
+
+			// Before we wait, clear signal wait state. This must happen while locked. Otherwise we would clear the Event trigger from an simultaneously running Release()
+			Event->Reset();
+
+			// Release the lock and wait. The event could be triggered already, but the Event->Reset() was protected by the lock and we get signaled for any previous trigger
+			CriticalSection.Unlock();
+			if (!Event->Wait())
+			{
+				return false;
+			}
+			CriticalSection.Lock();
+		}
+		CriticalSection.Unlock();
+		return true;
+	}
+
+	bool Obtain(int64 WaitMicroSeconds)
+	{
+		CriticalSection.Lock();
+
+		FTimespan WaitTime = FTimespan::FromMicroseconds(WaitMicroSeconds);
+		int64 TicksStart = FDateTime::Now().GetTicks();
+
+		// Check if decrement operation would work.
+		while(FPlatformAtomics::InterlockedDecrement(&Count) <= -1)
+		{
+			// Undo our decrement operations, because we know this was wrong...
+			FPlatformAtomics::InterlockedIncrement(&Count);
+
+			// Before we wait, clear signal wait state. This must happen while locked. Otherwise we would clear the Event trigger from an simultaneously running Release()
+			Event->Reset();
+
+			// Release the lock and wait. The event could be triggered already, but the Event->Reset() was protected by the lock and we get signaled for any previous trigger
+			CriticalSection.Unlock();
+
+			// For some strange reason, no waiting time left anymore... Exit!
+			if (WaitTime <= FTimespan::Zero())
+			{
+				return false;
+			}
+
+			// Check if the wait timed out...
+			if (!Event->Wait(WaitTime))
+			{
+				// Time out. Nothing can safe us here...
+				return false;
+			}
+
+			// Event was triggered.
+			CriticalSection.Lock();
+
+			// Calculate remaining waiting time as we could be triggered, BUT still will not get the value, because some other thread grabbed it first!
+			int64 TicksUntilHere = FDateTime::Now().GetTicks();
+			FTimespan DeltaWaiting(TicksUntilHere - TicksStart);
+			// Decrement time we are allowed to wait
+			WaitTime -= DeltaWaiting;
+			TicksStart = TicksUntilHere;
+		}
+		CriticalSection.Unlock();
+		return true;
+	}
+
+	bool TryToObtain()
+	{
+		FScopeLock Lock(&CriticalSection);
+		// Try if decrement operation would work. We inspect the value before the decrment operation...
+		if (FPlatformAtomics::InterlockedDecrement(&Count) <= -1)
+		{
+			// Nope: Undo the "wrong" decrement operation to restore value...
+			FPlatformAtomics::InterlockedIncrement(&Count);
+			return false;
+		}
+		return true;
+	}
+
+	void Release()
+	{
+		FScopeLock Lock(&CriticalSection);
+		FPlatformAtomics::InterlockedIncrement(&Count);
+		Event->Trigger();
+	}
+
+	int32 CurrentCount() const
+	{
+		FScopeLock Lock(&CriticalSection);
+		return Count;
+	}
+
+	void SetCount(int32 NewCurrentCount)
+	{
+		FScopeLock Lock(&CriticalSection);
+		Count = NewCurrentCount;
+	}
+
+private:
+	mutable FCriticalSection CriticalSection;
+	mutable FEvent* Event;
+	int32 Count;
+};
diff --git a/Engine/Plugins/Media/ElectraUtil/Source/ElectraBase/Public/Core/MediaThreads.h b/Engine/Plugins/Media/ElectraUtil/Source/ElectraBase/Public/Core/MediaThreads.h
new file mode 100644
--- /dev/null
+++ b/Engine/Plugins/Media/ElectraUtil/Source/ElectraBase/Public/Core/MediaThreads.h
@@ -0,0 +1,166 @@
+// Copyright Epic Games, Inc. All Rights Reserved.
+
+#pragma once
+
+#include "Core/MediaTypes.h"
+#include "Core/MediaMacros.h"
+#include "Core/MediaNoncopyable.h"
+#include "Core/MediaEventSignal.h"
+
+#include "HAL/PlatformAffinity.h"
+#include "HAL/Runnable.h"
+#include "HAL/PlatformProcess.h"
+#include "HAL/CriticalSection.h"
+#include "Delegates/Delegate.h"
+
+#include "Containers/UnrealString.h"
+#include "Containers/StringConv.h"
+#include "Templates/SharedPointer.h"
+#include "Templates/Function.h"
+
+
+class ELECTRABASE_API FMediaRunnable : private TMediaNoncopyable<FMediaRunnable>, public FRunnable
+{
+public:
+	DECLARE_DELEGATE(FStartDelegate);
+
+	/**
+	 * Common thread configuration parameters.
+	 */
+	struct Param
+	{
+		// Set some standard values as defaults.
+		EThreadPriority Priority = TPri_Normal;
+		uint32 StackSize = 65536;;
+		int32 CoreAffinity = -1;
+	};
+
+	static void Startup();
+	static void Shutdown();
+	static void EnqueueAsyncTask(TFunction<void()>&& InFunctionToExecuteOnAsyncThread);
+
+	static FMediaRunnable* Create(int32 CoreAffinityMask, EThreadPriority Priority, uint32 StackSize, const FString& InThreadName);
+	static void Destroy(FMediaRunnable* Thread);
+
+	void Start(FStartDelegate Entry, bool bWaitRunning = false);
+
+	void SetDoneSignal(FMediaEvent* DoneSignal);
+
+	void SetName(const FString& InThreadName);
+
+	EThreadPriority ChangePriority(EThreadPriority NewPriority);
+
+	EThreadPriority PriorityGet() const
+	{
+		return ThreadPriority;
+	}
+
+	uint32 StackSizeGet() const
+	{
+		return StackSize;
+	}
+
+	static uint32 StackSizeGetDefault()
+	{
+		return 0;
+	}
+
+	static void SleepSeconds(uint32 Seconds)
+	{
+		FPlatformProcess::Sleep((float)Seconds);
+	}
+
+	static void SleepMilliseconds(uint32 Milliseconds)
+	{
+		FPlatformProcess::Sleep(Milliseconds / 1000.0f);
+	}
+
+	static void SleepMicroseconds(uint32 Microseconds)
+	{
+		FPlatformProcess::Sleep(Microseconds / 1000000.0f);
+	}
+
+private:
+	FCriticalSection StateAccessMutex;
+	FStartDelegate EntryFunction;
+	FString ThreadName;
+	FRunnableThread* MediaThreadRunnable = nullptr;
+	FEvent* SignalRunning = nullptr;
+	FMediaEvent* DoneSignal = nullptr;
+	EThreadPriority ThreadPriority = TPri_Normal;
+	uint64 InitialCoreAffinity = 0;
+	uint32 StackSize = 0;
+	bool bIsStarted = false;
+
+	FMediaRunnable();
+	~FMediaRunnable();
+
+	void StartInternal();
+	uint32 Run() override;
+	void Exit() override;
+};
+
+
+/**
+ * A thread base class to either inherit from or use as a variable.
+ *
+ * Thread parameters are given to constructor but can be changed before starting the thread
+ * with the ThreadSet..() functions.
+ * This allows the constructor to use system default values so an instance of this class
+ * can be used as a member variable of some class.
+ *
+ * To start the thread on some function of type void(*)(void) call ThreadStart() with
+ * an appropriate delegate.
+ *
+ * IMPORTANT: If used as a member variable you have to ensure for your class to stay alive
+ *            i.e. not be destroyed before the thread function has finished. Otherwise your
+ *            function will most likely crash.
+ *            Either wait for thread completion by calling ThreadWaitDone() *OR*
+ *            have the destructor wait for the thread itself by calling ThreadWaitDoneOnDelete(true)
+ *            at some point prior to destruction - preferably before starting the thread.
+ */
+class ELECTRABASE_API FMediaThread : private TMediaNoncopyable<FMediaThread>
+{
+public:
+	virtual ~FMediaThread();
+
+	// Default constructor. Uses system defaults if no values given.
+	FMediaThread(const char* AnsiName = nullptr);
+
+	// Set a thread priority other than the one given to the constructor before starting the thread.
+	void ThreadSetPriority(EThreadPriority Priority);
+
+	// Set a core affinity before starting the thread. Defaults to -1 for no affinity (run on any core).
+	void ThreadSetCoreAffinity(int32 CoreAffinity);
+
+	// Set a thread stack size other than the one given to the constructor before starting the thread.
+	void ThreadSetStackSize(uint32 StackSize);
+
+	// Set a thread name other than the one given to the constructor before starting the thread.
+	void ThreadSetName(const char* InAnsiThreadName);
+
+	// Sets a new thread name once the thread is running.
+	void ThreadRename(const char* InAnsiThreadName);
+
+	// Sets whether or not the destructor needs to wait for the thread to have finished. Defaults to false. Useful when using this as a member variable.
+	void ThreadWaitDoneOnDelete(bool bWait);
+
+	// Waits for the thread to have finished.
+	void ThreadWaitDone();
+
+	// Starts the thread at the given function void(*)(void)
+	void ThreadStart(FMediaRunnable::FStartDelegate EntryFunction);
+
+	// Resets the thread to be started again. Must have waited for thread termination using ThreadWaitDone() first!
+	void ThreadReset();
+
+private:
+	FMediaEvent SigDone;
+	FString ThreadName;
+	FMediaRunnable* MediaRunnable = nullptr;
+	EThreadPriority Priority = TPri_Normal;
+	uint32 StackSize = 0;
+	int32 CoreAffinity = 0;
+	bool bIsStarted = false;;
+	bool bWaitDoneOnDelete = false;
+};
diff --git a/Engine/Plugins/Media/ElectraUtil/Source/ElectraBase/Public/Core/MediaTypes.h b/Engine/Plugins/Media/ElectraUtil/Source/ElectraBase/Public/Core/MediaTypes.h
new file mode 100644
--- /dev/null
+++ b/Engine/Plugins/Media/ElectraUtil/Source/ElectraBase/Public/Core/MediaTypes.h
@@ -0,0 +1,44 @@
+// Copyright Epic Games, Inc. All Rights Reserved.
+
+#pragma once
+
+#include "Misc/Build.h"
+
+// Include low level includes. This will also include HAL/Platform.h for our types
+#include "CoreTypes.h"
+#include "Templates/SharedPointer.h"
+
+#if PLATFORM_MICROSOFT
+	#include "Microsoft/WindowsHWrapper.h"
+#endif
+
+typedef uint32  UEMediaError;
+
+#define UEMEDIA_ERROR_OK				(0)
+#define UEMEDIA_ERROR_DETAIL			(1)
+#define UEMEDIA_ERROR_INTERNAL			(2)
+#define UEMEDIA_ERROR_NOT_IMPLEMENTED	(3)
+#define UEMEDIA_ERROR_ABORTED			(4)
+#define UEMEDIA_ERROR_END_OF_STREAM		(5)
+#define UEMEDIA_ERROR_INSUFFICIENT_DATA	(6)
+#define UEMEDIA_ERROR_READ_ERROR		(7)
+#define UEMEDIA_ERROR_FORMAT_ERROR		(8)
+#define UEMEDIA_ERROR_NOT_SUPPORTED		(9)
+#define UEMEDIA_ERROR_TRY_AGAIN			(100)
+#define UEMEDIA_ERROR_OOM				(0x80000000U)
+#define UEMEDIA_ERROR_BAD_ARGUMENTS		(0x80000001U)
+
+namespace Electra
+{
+	template <typename T>
+	using TSharedPtrTS = TSharedPtr<T, ESPMode::ThreadSafe>;
+
+	template <typename T>
+	using TWeakPtrTS = TWeakPtr<T, ESPMode::ThreadSafe>;
+
+	template<typename T, typename... ArgTypes>
+	TSharedRef<T, ESPMode::ThreadSafe> MakeSharedTS(ArgTypes&&... Args)
+	{
+		return MakeShared<T, ESPMode::ThreadSafe>(Forward<ArgTypes>(Args)...);
+	}
+}
