Better handling of a failure to acquire the shared resource sync in D3D11

#rnx

diff --git a/Engine/Plugins/Media/ElectraUtil/Source/ElectraSamples/Private/ElectraTextureSamplePC.cpp b/Engine/Plugins/Media/ElectraUtil/Source/ElectraSamples/Private/ElectraTextureSamplePC.cpp
--- a/Engine/Plugins/Media/ElectraUtil/Source/ElectraSamples/Private/ElectraTextureSamplePC.cpp
+++ b/Engine/Plugins/Media/ElectraUtil/Source/ElectraSamples/Private/ElectraTextureSamplePC.cpp
@@ -5,6 +5,7 @@
 #if PLATFORM_WINDOWS && !UE_SERVER
 
 #include "ElectraTextureSample.h"
+#include "ElectraSamplesModule.h"
 
 #include "ProfilingDebugging/RealtimeGPUProfiler.h"
 #include "RenderUtils.h"
@@ -312,13 +313,14 @@ struct FRHICommandCopyResourceDX11 final : public FRHICommand<FRHICommandCopyRes
 										// Key is 1 : Texture as just been updated
 										// Key is 2 : Texture as already been updated.
 										// Do not wait to acquire key 1 since there is race no condition between writer and reader.
-										if (KeyedMutex->AcquireSync(1, 0) == S_OK)
+										HRESULT Result = KeyedMutex->AcquireSync(1, 0);
+										if (Result == S_OK)
 										{
 											// Copy from shared texture of FSink device to Rendering device
 											D3D11DeviceContext->CopyResource(DestinationTexture, SharedResource);
 											KeyedMutex->ReleaseSync(2);
 										}
-										else
+										else if (Result == ((HRESULT)WAIT_TIMEOUT))
 										{
 											// If key 1 cannot be acquired, another reader is already copying the resource
 											// and will release key with 2. 
@@ -328,6 +330,10 @@ struct FRHICommandCopyResourceDX11 final : public FRHICommand<FRHICommandCopyRes
 												KeyedMutex->ReleaseSync(2);
 											}
 										}
+										else
+										{
+											UE_LOG(LogElectraSamples, Warning, TEXT("AcquireSync failed with 0x%08x!"), Result);
+										}
 									}
 								}
 							}
