[Electra]
* Moved dllstorage from type to functions to remove linker warnings when building with merged modules

diff --git a/Engine/Plugins/Media/ElectraCodecs/Source/ElectraDecoders/Public/Utils/ElectraBitstreamReader.h b/Engine/Plugins/Media/ElectraCodecs/Source/ElectraDecoders/Public/Utils/ElectraBitstreamReader.h
--- a/Engine/Plugins/Media/ElectraCodecs/Source/ElectraDecoders/Public/Utils/ElectraBitstreamReader.h
+++ b/Engine/Plugins/Media/ElectraCodecs/Source/ElectraDecoders/Public/Utils/ElectraBitstreamReader.h
@@ -11,7 +11,7 @@
 namespace ElectraDecodersUtil
 {
 
-class ELECTRADECODERS_API FElectraBitstreamReader
+class FElectraBitstreamReader
 {
 public:
 	FElectraBitstreamReader() = default;
diff --git a/Engine/Plugins/Media/ElectraUtil/Source/ElectraBase/Public/PlayerTime.h b/Engine/Plugins/Media/ElectraUtil/Source/ElectraBase/Public/PlayerTime.h
--- a/Engine/Plugins/Media/ElectraUtil/Source/ElectraBase/Public/PlayerTime.h
+++ b/Engine/Plugins/Media/ElectraUtil/Source/ElectraBase/Public/PlayerTime.h
@@ -13,7 +13,7 @@ class FTimeFraction;
 /**
  * Keeps a time value in hundred nanoseconds (HNS).
  */
-class ELECTRABASE_API FTimeValue
+class FTimeValue
 {
 public:
 	static FTimeValue GetInvalid()
@@ -144,7 +144,7 @@ public:
 	}
 
 	//! Returns this time value in a custom timebase. Requires internal bigint conversion and is therefor SLOW!
-	int64 GetAsTimebase(uint32 CustomTimebase) const;
+	ELECTRABASE_API int64 GetAsTimebase(uint32 CustomTimebase) const;
 
 	FTimespan GetAsTimespan() const
 	{
@@ -269,9 +269,9 @@ public:
 		return *this;
 	}
 
-	FTimeValue& SetFromND(int64 Numerator, uint32 Denominator, int64 InSequenceIndex=0);
+	ELECTRABASE_API FTimeValue& SetFromND(int64 Numerator, uint32 Denominator, int64 InSequenceIndex=0);
 
-	FTimeValue& SetFromTimeFraction(const FTimeFraction& TimeFraction, int64 InSequenceIndex=0);
+	ELECTRABASE_API FTimeValue& SetFromTimeFraction(const FTimeFraction& TimeFraction, int64 InSequenceIndex=0);
 
 
 	FTimeValue& SetFromTimespan(const FTimespan& InTimespan, int64 InSequenceIndex=0)
@@ -586,7 +586,7 @@ struct FTimeRange
 /**
  * Keeps a time value as a fractional.
  */
-class ELECTRABASE_API FTimeFraction
+class FTimeFraction
 {
 public:
 	static const FTimeFraction& GetInvalid()
@@ -653,7 +653,7 @@ public:
 	}
 
 	//! Returns this time value in a custom timebase. Requires internal bigint conversion and is therefor SLOW!
-	int64 GetAsTimebase(uint32 CustomTimebase) const;
+	ELECTRABASE_API int64 GetAsTimebase(uint32 CustomTimebase) const;
 
 	FTimeFraction& SetFromND(int64 InNumerator, uint32 InDenominator)
 	{
@@ -680,7 +680,7 @@ public:
 		return *this;
 	}
 
-	FTimeFraction& SetFromFloatString(const FString& In);
+	ELECTRABASE_API FTimeFraction& SetFromFloatString(const FString& In);
 
 
 	bool operator == (const FTimeFraction& rhs) const
