Stopped ProGuard from obfuscating a class used by JNI

#rnx

diff --git a/Engine/Plugins/Media/ElectraUtil/Source/ElectraSamples/ElectraSamples_UPL.xml b/Engine/Plugins/Media/ElectraUtil/Source/ElectraSamples/ElectraSamples_UPL.xml
--- a/Engine/Plugins/Media/ElectraUtil/Source/ElectraSamples/ElectraSamples_UPL.xml
+++ b/Engine/Plugins/Media/ElectraUtil/Source/ElectraSamples/ElectraSamples_UPL.xml
@@ -31,6 +31,9 @@
 }
 -keep class com.epicgames.unreal.ElectraTextureSample$FFrameUpdateInfo {
    public *;
+}
+-keep class com.epicgames.unreal.ElectraTextureSample$FImageResources {
+   public *;
 }
 		</insert>
 	</proguardAdditions>
diff --git a/Engine/Plugins/Media/ElectraUtil/Source/ElectraSamples/Private/Java/src/com/epicgames/unreal/ElectraTextureSample.java b/Engine/Plugins/Media/ElectraUtil/Source/ElectraSamples/Private/Java/src/com/epicgames/unreal/ElectraTextureSample.java
--- a/Engine/Plugins/Media/ElectraUtil/Source/ElectraSamples/Private/Java/src/com/epicgames/unreal/ElectraTextureSample.java
+++ b/Engine/Plugins/Media/ElectraUtil/Source/ElectraSamples/Private/Java/src/com/epicgames/unreal/ElectraTextureSample.java
@@ -20,7 +20,7 @@ public class ElectraTextureSample
 {
 	private static final String TAG = "[ElectraTextureSample]";
 
-	class FImageResources
+	public class FImageResources
 	{
 		public void Release()
 		{
