Moving creating texture from HardwareBuffer into RHI
#rnx

diff --git a/Engine/Plugins/Media/ElectraUtil/Source/ElectraSamples/Private/ElectraTextureSampleAndroid.cpp b/Engine/Plugins/Media/ElectraUtil/Source/ElectraSamples/Private/ElectraTextureSampleAndroid.cpp
--- a/Engine/Plugins/Media/ElectraUtil/Source/ElectraSamples/Private/ElectraTextureSampleAndroid.cpp
+++ b/Engine/Plugins/Media/ElectraUtil/Source/ElectraSamples/Private/ElectraTextureSampleAndroid.cpp
@@ -70,49 +70,6 @@ static TAutoConsoleVariable<int32> CVarElectraAndroidUseGpuOutputPath(
 
 namespace {
 
-struct FElectraTextureSampleVulkanResources
-{
-	VkImage Image;
-	VkDeviceMemory DeviceMemory;
-};
-
-static void CleanupImageResourcesVulkan(void* UserData)
-{
-	check(UserData);
-
-	FElectraTextureSampleVulkanResources* VulkanResources = static_cast<FElectraTextureSampleVulkanResources*>(UserData);
-
-	IVulkanDynamicRHI* RHI = GetIVulkanDynamicRHI();
-	VkDevice Device = RHI->RHIGetVkDevice();
-	const VkAllocationCallbacks* AllocationCallbacks = RHI->RHIGetVkAllocationCallbacks();
-
-	if (VulkanResources->DeviceMemory != VK_NULL_HANDLE)
-	{
-		static PFN_vkFreeMemory VkFreeMemory = nullptr;
-		if (!VkFreeMemory)
-		{
-			VkFreeMemory = (PFN_vkFreeMemory)RHI->RHIGetVkInstanceProcAddr("vkFreeMemory");
-			check(VkFreeMemory);
-		}
-
-		VkFreeMemory(Device, VulkanResources->DeviceMemory, AllocationCallbacks);
-	}
-
-	if (VulkanResources->Image != VK_NULL_HANDLE)
-	{
-		static PFN_vkDestroyImage VkDestroyImage = nullptr;
-		if (!VkDestroyImage)
-		{
-			VkDestroyImage = (PFN_vkDestroyImage)RHI->RHIGetVkInstanceProcAddr("vkDestroyImage");
-			check(VkDestroyImage);
-		}
-
-		VkDestroyImage(Device, VulkanResources->Image, AllocationCallbacks);
-	}
-
-	delete VulkanResources;
-}
-
 static void CleanupImageResourcesJNI(jobject Resources, jmethodID ReleaseFN)
 {
 	QUICK_SCOPE_CYCLE_COUNTER(STAT_FRHICommandReleaseDecoderResources_Execute);
@@ -183,9 +140,14 @@ public:
 	jobject ImageResources_GetHardwareBuffer(jobject ImageResources);
 	void ImageResources_GetScaleOffset(jobject ImageResources, FVector2f& OutScale, FVector2f& OutOffset);
 
-	const bool UseGpuOutputPath;
+	bool UseGpuOutputPath() const
+	{
+		return bUseGpuOutputPath;
+	}
 
 private:
+	bool bUseGpuOutputPath;
+
 	// Java methods
 	FJavaClassMethod	InitializeFN;
 	FJavaClassMethod	ReleaseFN;
@@ -255,7 +217,7 @@ private:
 
 FElectraTextureSampleSupport::FElectraTextureSampleSupport()
 	: FJavaClassObject(GetClassName(), "()V")
-	, UseGpuOutputPath(FAndroidMisc::ShouldUseVulkan() && (CVarElectraAndroidUseGpuOutputPath.GetValueOnAnyThread() != 0)) // GpuOutputPath is only available for Vulkan right now (and experimental)
+	, bUseGpuOutputPath(false)
 	, InitializeFN(GetClassMethod("Initialize", "(ZZJ)V"))
 	, ReleaseFN(GetClassMethod("Release", "()V"))
 	, GetCodecSurfaceFN(GetClassMethod("GetCodecSurface", "()Landroid/view/Surface;"))
@@ -265,6 +227,36 @@ FElectraTextureSampleSupport::FElectraTextureSampleSupport()
 	, CodecSurfaceToDelete(nullptr)
 	, CodecSurfaceReadEvent(nullptr)
 {
+	if (FAndroidMisc::ShouldUseVulkan()) // GpuOutputPath is only available for Vulkan right now (and experimental)
+	{
+		if (CVarElectraAndroidUseGpuOutputPath.GetValueOnAnyThread() == 0) 
+		{
+			UE_LOG(LogElectraSamples, Log, TEXT("Selecting CPU path because GPU path is disabled via Electra.AndroidUseGpuOutputPath = 0"));
+		}
+		else
+		{
+			IVulkanDynamicRHI* RHI = GetIVulkanDynamicRHI();
+			TArray<FAnsiString> LoadedDeviceExtensions = RHI->RHIGetLoadedDeviceExtensions();
+			if (!LoadedDeviceExtensions.Find(VK_KHR_SAMPLER_YCBCR_CONVERSION_EXTENSION_NAME))
+			{
+				UE_LOG(LogElectraSamples, Log, TEXT("Selecting CPU path because GPU extension '" VK_KHR_SAMPLER_YCBCR_CONVERSION_EXTENSION_NAME "' is not available!"));
+			}
+			else if (!LoadedDeviceExtensions.Find(VK_ANDROID_EXTERNAL_MEMORY_ANDROID_HARDWARE_BUFFER_EXTENSION_NAME))
+			{
+				UE_LOG(LogElectraSamples, Log, TEXT("Selecting CPU path because GPU extension '" VK_ANDROID_EXTERNAL_MEMORY_ANDROID_HARDWARE_BUFFER_EXTENSION_NAME "' is not available!"));
+			}
+			else
+			{
+				bUseGpuOutputPath = true;
+				UE_LOG(LogElectraSamples, Log, TEXT("Selecting GPU path because it is enabled via Electra.AndroidUseGpuOutputPath = 1"));
+			}
+		}
+	}
+	else
+	{
+		UE_LOG(LogElectraSamples, Log, TEXT("Selecting CPU path because we are on OES"));
+	}
+
 	JNIEnv* JEnv = FAndroidApplication::GetJavaEnv();
 
 	// Get field IDs for FFrameUpdateInfo class members
@@ -334,7 +326,7 @@ FElectraTextureSampleSupport::FElectraTextureSampleSupport()
 				{
 					JNIEnv* JEnv = FAndroidApplication::GetJavaEnv();
 					// Setup Java side of things
-					JEnv->CallVoidMethod(Object, InitializeFN.Method, UseGpuOutputPath, FAndroidMisc::ShouldUseVulkan(), (jlong)this);
+					JEnv->CallVoidMethod(Object, InitializeFN.Method, UseGpuOutputPath(), FAndroidMisc::ShouldUseVulkan(), (jlong)this);
 					// Query surface to be used for decoder
 					jobject Surface = JEnv->CallObjectMethod(Object, GetCodecSurfaceFN.Method);
 					CodecSurface = JEnv->NewGlobalRef(Surface);
@@ -496,7 +488,7 @@ int32 FElectraTextureSampleSupport::GetFrameData(FElectraTextureSample* InTarget
 	{
 		if (InTargetSample)
 		{
-			if (UseGpuOutputPath)
+			if (UseGpuOutputPath())
 			{
 				jobject ImageResources = JEnv->GetObjectField(OutputInfo, FFrameUpdateInfo_ImageResources);
 				if (ImageResources)
@@ -612,7 +604,7 @@ void FElectraTextureSample::Initialize(FVideoDecoderOutput* InVideoDecoderOutput
 
 	ensure(VideoDecoderOutputAndroid->GetOutputType() == FVideoDecoderOutputAndroid::EOutputType::DirectToSurfaceAsQueue);
 
-	if (Support->UseGpuOutputPath)
+	if (Support->UseGpuOutputPath())
 	{
 		Support->GetFrameData(this);
 		Texture = nullptr;
@@ -742,133 +734,7 @@ FTextureRHIRef FElectraTextureSample::InitializeTextureVulkan(AHardwareBuffer* H
 	check(FAndroidMisc::ShouldUseVulkan());
 
 	IVulkanDynamicRHI* RHI = GetIVulkanDynamicRHI();
-	VkDevice Device = RHI->RHIGetVkDevice();
-	const VkAllocationCallbacks* AllocationCallbacks = RHI->RHIGetVkAllocationCallbacks();
-
-	AHardwareBuffer_Desc Desc;
-	AHardwareBuffer_describe(HardwareBuffer, &Desc);
-	check((Desc.usage & AHARDWAREBUFFER_USAGE_GPU_SAMPLED_IMAGE) != 0);
-
-	uint32 SizeX = Desc.width;
-	uint32 SizeY = Desc.height;
-	uint32 NumMips = 1;
-	uint32 NumSamples = 1;
-
-	ETextureCreateFlags Flags = ETextureCreateFlags::External;
-	const FClearValueBinding& ClearValueBinding = FClearValueBinding::Transparent;
-
-	static PFN_vkGetAndroidHardwareBufferPropertiesANDROID VkGetAndroidHardwareBufferPropertiesANDROID = nullptr;
-	if (!VkGetAndroidHardwareBufferPropertiesANDROID)
-	{
-		VkGetAndroidHardwareBufferPropertiesANDROID = (PFN_vkGetAndroidHardwareBufferPropertiesANDROID)RHI->RHIGetVkInstanceProcAddr("vkGetAndroidHardwareBufferPropertiesANDROID");
-		check(VkGetAndroidHardwareBufferPropertiesANDROID);
-	}
-
-	VkAndroidHardwareBufferFormatPropertiesANDROID HardwareBufferFormatProperties;
-	ZeroVulkanStruct(HardwareBufferFormatProperties, VK_STRUCTURE_TYPE_ANDROID_HARDWARE_BUFFER_FORMAT_PROPERTIES_ANDROID);
-
-	VkAndroidHardwareBufferPropertiesANDROID HardwareBufferProperties;
-	ZeroVulkanStruct(HardwareBufferProperties, VK_STRUCTURE_TYPE_ANDROID_HARDWARE_BUFFER_PROPERTIES_ANDROID);
-	HardwareBufferProperties.pNext = &HardwareBufferFormatProperties;
-
-	VkResult Result = VkGetAndroidHardwareBufferPropertiesANDROID(Device, HardwareBuffer, &HardwareBufferProperties);
-	check(Result == VK_SUCCESS);
-
-	VkExternalFormatANDROID ExternalFormat;
-	ZeroVulkanStruct(ExternalFormat, VK_STRUCTURE_TYPE_EXTERNAL_FORMAT_ANDROID);
-	ExternalFormat.externalFormat = HardwareBufferFormatProperties.externalFormat;
-
-	VkExternalMemoryImageCreateInfo ExternalMemoryImageCreateInfo;
-	ZeroVulkanStruct(ExternalMemoryImageCreateInfo, VK_STRUCTURE_TYPE_EXTERNAL_MEMORY_IMAGE_CREATE_INFO);
-	ExternalMemoryImageCreateInfo.handleTypes = VK_EXTERNAL_MEMORY_HANDLE_TYPE_ANDROID_HARDWARE_BUFFER_BIT_ANDROID;
-	ExternalMemoryImageCreateInfo.pNext = &ExternalFormat;
-
-	VkImageCreateInfo ImageCreateInfo;
-	ZeroVulkanStruct(ImageCreateInfo, VK_STRUCTURE_TYPE_IMAGE_CREATE_INFO);
-	ImageCreateInfo.imageType = VK_IMAGE_TYPE_2D;
-	ImageCreateInfo.format = VK_FORMAT_UNDEFINED;
-	ImageCreateInfo.extent.width = SizeX;
-	ImageCreateInfo.extent.height = SizeY;
-	ImageCreateInfo.extent.depth = 1;
-	ImageCreateInfo.mipLevels = NumMips;
-	ImageCreateInfo.arrayLayers = Desc.layers;
-
-	ImageCreateInfo.flags = 0;
-	ImageCreateInfo.samples = VK_SAMPLE_COUNT_1_BIT;
-
-	ImageCreateInfo.tiling = VK_IMAGE_TILING_OPTIMAL;
-	ImageCreateInfo.usage = VK_IMAGE_USAGE_SAMPLED_BIT;
-	ImageCreateInfo.sharingMode = VK_SHARING_MODE_EXCLUSIVE;
-	ImageCreateInfo.queueFamilyIndexCount = 0;
-	ImageCreateInfo.pQueueFamilyIndices = nullptr;
-	ImageCreateInfo.initialLayout = VK_IMAGE_LAYOUT_UNDEFINED;
-
-	ImageCreateInfo.pNext = &ExternalMemoryImageCreateInfo;
-
-	static PFN_vkCreateImage VkCreateImage = nullptr;
-	if (!VkCreateImage)
-	{
-		VkCreateImage = (PFN_vkCreateImage)RHI->RHIGetVkInstanceProcAddr("vkCreateImage");
-		check(VkCreateImage);
-	}
-
-	VkImage VulkanImage;
-	Result = VkCreateImage(Device, &ImageCreateInfo, nullptr, &VulkanImage);
-	check(Result == VK_SUCCESS);
-
-	VkMemoryDedicatedAllocateInfo MemoryDedicatedAllocateInfo;
-	ZeroVulkanStruct(MemoryDedicatedAllocateInfo, VK_STRUCTURE_TYPE_MEMORY_DEDICATED_ALLOCATE_INFO);
-	MemoryDedicatedAllocateInfo.image = VulkanImage;
-	MemoryDedicatedAllocateInfo.buffer = VK_NULL_HANDLE;
-
-	VkImportAndroidHardwareBufferInfoANDROID ImportAndroidHardwareBufferInfo;
-	ZeroVulkanStruct(ImportAndroidHardwareBufferInfo, VK_STRUCTURE_TYPE_IMPORT_ANDROID_HARDWARE_BUFFER_INFO_ANDROID);
-	ImportAndroidHardwareBufferInfo.buffer = HardwareBuffer;
-	ImportAndroidHardwareBufferInfo.pNext = &MemoryDedicatedAllocateInfo;
-
-	uint32 MemoryTypeBits = HardwareBufferProperties.memoryTypeBits;
-	check(MemoryTypeBits > 0); // No index available, this should never happen
-	uint32 MemoryTypeIndex = 0;
-	for (;(MemoryTypeBits & 1) != 1; ++MemoryTypeIndex)
-	{
-		MemoryTypeBits >>= 1;
-	}
-
-	VkMemoryAllocateInfo MemoryAllocateInfo;
-	ZeroVulkanStruct(MemoryAllocateInfo, VK_STRUCTURE_TYPE_MEMORY_ALLOCATE_INFO);
-	MemoryAllocateInfo.allocationSize = HardwareBufferProperties.allocationSize;
-	MemoryAllocateInfo.memoryTypeIndex = MemoryTypeIndex;
-	MemoryAllocateInfo.pNext = &ImportAndroidHardwareBufferInfo;
-
-	static PFN_vkAllocateMemory VkAllocateMemory = nullptr;
-	if (!VkAllocateMemory)
-	{
-		VkAllocateMemory = (PFN_vkAllocateMemory)RHI->RHIGetVkInstanceProcAddr("vkAllocateMemory");
-		check(VkAllocateMemory);
-	}
-
-	VkDeviceMemory VulkanDeviceMemory;
-	Result = VkAllocateMemory(Device, &MemoryAllocateInfo, AllocationCallbacks, &VulkanDeviceMemory);
-	check(Result == VK_SUCCESS);
-
-	static PFN_vkBindImageMemory VkBindImageMemory = nullptr;
-	if (!VkBindImageMemory)
-	{
-		VkBindImageMemory = (PFN_vkBindImageMemory)RHI->RHIGetVkInstanceProcAddr("vkBindImageMemory");
-		check(VkBindImageMemory);
-	}
-
-	Result = VkBindImageMemory(Device, VulkanImage, VulkanDeviceMemory, 0);
-	check(Result == VK_SUCCESS);
-
-	FElectraTextureSampleVulkanResources* ElectraTextureSampleVulkanResources = new FElectraTextureSampleVulkanResources{ VulkanImage, VulkanDeviceMemory };
-	FVulkanRHIExternalImageDeleteCallbackInfo ExternalImageDeleteCallbackInfo =
-	{
-		ElectraTextureSampleVulkanResources,
-		CleanupImageResourcesVulkan
-	};
-
-	return RHI->RHICreateTexture2DFromResource(PF_Unknown, SizeX, SizeY, NumMips, NumSamples, VulkanImage, Flags, ClearValueBinding, ExternalImageDeleteCallbackInfo);
+	return RHI->RHICreateTexture2DFromAndroidHardwareBuffer(HardwareBuffer);
 }
 
 
@@ -1126,7 +992,7 @@ void  FElectraTextureSample::CopyFromExternalTextureVulkan(FRHICommandListImmedi
 
 bool FElectraTextureSample::Convert(FRHICommandListImmediate& RHICmdList, FTextureRHIRef& InDstTexture, const FConversionHints& Hints)
 {
-	if (Support->UseGpuOutputPath)
+	if (Support->UseGpuOutputPath())
 		return ConvertGpuOutputPath(RHICmdList, InDstTexture, Hints);
 	else
 		return ConvertCpuOutputPath(RHICmdList, InDstTexture, Hints);
diff --git a/Engine/Source/Runtime/VulkanRHI/Private/Android/VulkanAndroidPlatform.h b/Engine/Source/Runtime/VulkanRHI/Private/Android/VulkanAndroidPlatform.h
--- a/Engine/Source/Runtime/VulkanRHI/Private/Android/VulkanAndroidPlatform.h
+++ b/Engine/Source/Runtime/VulkanRHI/Private/Android/VulkanAndroidPlatform.h
@@ -32,6 +32,7 @@
 
 #define ENUM_VK_ENTRYPOINTS_PLATFORM_INSTANCE(EnumMacro) \
 	EnumMacro(PFN_vkCreateAndroidSurfaceKHR, vkCreateAndroidSurfaceKHR) \
+	EnumMacro(PFN_vkGetAndroidHardwareBufferPropertiesANDROID, vkGetAndroidHardwareBufferPropertiesANDROID)
 
 #define ENUM_VK_ENTRYPOINTS_OPTIONAL_PLATFORM_INSTANCE(EnumMacro) \
 	EnumMacro(PFN_vkGetRefreshCycleDurationGOOGLE, vkGetRefreshCycleDurationGOOGLE) \
diff --git a/Engine/Source/Runtime/VulkanRHI/Private/VulkanCommandWrappers.h b/Engine/Source/Runtime/VulkanRHI/Private/VulkanCommandWrappers.h
--- a/Engine/Source/Runtime/VulkanRHI/Private/VulkanCommandWrappers.h
+++ b/Engine/Source/Runtime/VulkanRHI/Private/VulkanCommandWrappers.h
@@ -250,6 +250,12 @@ struct FWrapLayer
 	static void GetDeferredOperationMaxConcurrencyKHR(VkResult Result, VkDevice Device, VkDeferredOperationKHR DeferredOperation) VULKAN_LAYER_BODY
 	static void GetDeferredOperationResultKHR(VkResult Result, VkDevice Device, VkDeferredOperationKHR DeferredOperation) VULKAN_LAYER_BODY
 	static void GetDeviceFaultInfoEXT(VkResult Result, VkDevice Device, VkDeviceFaultCountsEXT* FaultCounts, VkDeviceFaultInfoEXT* FaultInfo) VULKAN_LAYER_BODY
+	static void CreateSamplerYcbcrConversion(VkResult Result, VkDevice Device, const VkSamplerYcbcrConversionCreateInfo* CreateInfo, const VkAllocationCallbacks* Allocator, VkSamplerYcbcrConversion* YcbcrConversion) VULKAN_LAYER_BODY
+	static void DestroySamplerYcbcrConversion(VkResult Result, VkDevice Device, VkSamplerYcbcrConversion YcbcrConversion, const VkAllocationCallbacks* Allocator) VULKAN_LAYER_BODY
+#if PLATFORM_ANDROID
+	static void GetAndroidHardwareBufferPropertiesANDROID(VkResult Result, VkDevice Device, const struct AHardwareBuffer* Buffer, VkAndroidHardwareBufferPropertiesANDROID* Properties) VULKAN_LAYER_BODY
+#endif
+
 };
 
 #undef VULKAN_LAYER_BODY
@@ -1732,6 +1738,31 @@ namespace VulkanRHI
 		return Result;
 	}
 
+	static FORCEINLINE_DEBUGGABLE VkResult vkCreateSamplerYcbcrConversion(VkDevice Device, const VkSamplerYcbcrConversionCreateInfo* CreateInfo, const VkAllocationCallbacks* Allocator, VkSamplerYcbcrConversion* YcbcrConversion)
+	{
+		FWrapLayer::CreateSamplerYcbcrConversion(VK_RESULT_MAX_ENUM, Device, CreateInfo, Allocator, YcbcrConversion);
+		VkResult Result = VULKANAPINAMESPACE::vkCreateSamplerYcbcrConversion(Device, CreateInfo, Allocator, YcbcrConversion);
+		FWrapLayer::CreateSamplerYcbcrConversion(Result, Device, CreateInfo, Allocator, YcbcrConversion);
+		return Result;
+	}
+
+	static FORCEINLINE_DEBUGGABLE void vkDestroySamplerYcbcrConversion(VkDevice Device, VkSamplerYcbcrConversion YcbcrConversion, const VkAllocationCallbacks* Allocator)
+	{
+		FWrapLayer::DestroySamplerYcbcrConversion(VK_RESULT_MAX_ENUM, Device, YcbcrConversion, Allocator);
+		VULKANAPINAMESPACE::vkDestroySamplerYcbcrConversion(Device, YcbcrConversion, Allocator);
+		FWrapLayer::DestroySamplerYcbcrConversion(VK_SUCCESS, Device, YcbcrConversion, Allocator);
+	}
+
+#if PLATFORM_ANDROID
+	static FORCEINLINE_DEBUGGABLE VkResult vkGetAndroidHardwareBufferPropertiesANDROID(VkDevice Device, const struct AHardwareBuffer* Buffer, VkAndroidHardwareBufferPropertiesANDROID* Properties)
+	{
+		FWrapLayer::GetAndroidHardwareBufferPropertiesANDROID(VK_RESULT_MAX_ENUM, Device, Buffer, Properties);
+		VkResult Result = VULKANAPINAMESPACE::vkGetAndroidHardwareBufferPropertiesANDROID(Device, Buffer, Properties);
+		FWrapLayer::GetAndroidHardwareBufferPropertiesANDROID(Result, Device, Buffer, Properties);
+		return Result;
+	}
+#endif // PLATFORM_ANDROID
+
 #if VULKAN_ENABLE_IMAGE_TRACKING_LAYER
 	void BindDebugLabelName(VkImage Image, const TCHAR* Name);
 #endif
diff --git a/Engine/Source/Runtime/VulkanRHI/Private/VulkanDebug.cpp b/Engine/Source/Runtime/VulkanRHI/Private/VulkanDebug.cpp
--- a/Engine/Source/Runtime/VulkanRHI/Private/VulkanDebug.cpp
+++ b/Engine/Source/Runtime/VulkanRHI/Private/VulkanDebug.cpp
@@ -4570,6 +4570,37 @@ void FWrapLayer::QueueSubmit2KHR(VkResult Result, VkQueue Queue, uint32_t Submit
 	}
 }
 
+void FWrapLayer::CreateSamplerYcbcrConversion(VkResult Result, VkDevice Device, const VkSamplerYcbcrConversionCreateInfo* CreateInfo, const VkAllocationCallbacks* Allocator, VkSamplerYcbcrConversion* YcbcrConversion)
+{
+	if (Result == VK_RESULT_MAX_ENUM)
+	{
+#if VULKAN_ENABLE_DUMP_LAYER
+		PrintfBeginResult(FString::Printf(TEXT("CreateSamplerYcbcrConversion(Device=0x%p, CreateInfo=0x%p, Allocator=0x%p, YcbcrConversion=0x%p))"), Device, CreateInfo, Allocator, YcbcrConversion));
+#endif
+	}
+}
+
+void FWrapLayer::DestroySamplerYcbcrConversion(VkResult Result, VkDevice Device, VkSamplerYcbcrConversion YcbcrConversion, const VkAllocationCallbacks* Allocator)
+{
+	if (Result == VK_RESULT_MAX_ENUM)
+	{
+#if VULKAN_ENABLE_DUMP_LAYER
+		PrintfBeginResult(FString::Printf(TEXT("DestroySamplerYcbcrConversion(Device=0x%p, YcbcrConversion=0x%p, Allocator=0x%p))"), Device, YcbcrConversion, Allocator));
+#endif
+	}
+}
+
+#if PLATFORM_ANDROID
+void FWrapLayer::GetAndroidHardwareBufferPropertiesANDROID(VkResult Result, VkDevice Device, const struct AHardwareBuffer* Buffer, VkAndroidHardwareBufferPropertiesANDROID* Properties)
+{
+	if (Result == VK_RESULT_MAX_ENUM)
+	{
+#if VULKAN_ENABLE_DUMP_LAYER
+		PrintfBeginResult(FString::Printf(TEXT("GetAndroidHardwareBufferPropertiesANDROID(Device=0x%p, Buffer=0x%p, Properties=0x%p))"), Device, Buffer, Properties));
+#endif
+	}
+}
+#endif
 
 #if VULKAN_ENABLE_IMAGE_TRACKING_LAYER
 namespace VulkanRHI
diff --git a/Engine/Source/Runtime/VulkanRHI/Private/VulkanDevice.h b/Engine/Source/Runtime/VulkanRHI/Private/VulkanDevice.h
--- a/Engine/Source/Runtime/VulkanRHI/Private/VulkanDevice.h
+++ b/Engine/Source/Runtime/VulkanRHI/Private/VulkanDevice.h
@@ -67,6 +67,7 @@ struct FOptionalVulkanDeviceExtensions
 			// Promoted to 1.1
 			uint64 HasKHRMultiview : 1;
 			uint64 HasKHR16bitStorage : 1;
+			uint64 HasKHRSamplerYcbcrConversion : 1;
 
 			// Promoted to 1.2
 			uint64 HasKHRRenderPass2 : 1;
diff --git a/Engine/Source/Runtime/VulkanRHI/Private/VulkanExtensions.cpp b/Engine/Source/Runtime/VulkanRHI/Private/VulkanExtensions.cpp
--- a/Engine/Source/Runtime/VulkanRHI/Private/VulkanExtensions.cpp
+++ b/Engine/Source/Runtime/VulkanRHI/Private/VulkanExtensions.cpp
@@ -1490,6 +1490,40 @@ public:
 	VkPhysicalDevicePipelineCreationCacheControlFeatures PhysicalDevicePipelineCreationCacheControlFeatures;
 };
 
+// ***** VK_KHR_sampler_ycbcr_conversion
+class FVulkanKHRSamplerYcbcrConversionExtension : public FVulkanDeviceExtension
+{
+public:
+
+	FVulkanKHRSamplerYcbcrConversionExtension(FVulkanDevice* InDevice) 
+		: FVulkanDeviceExtension(InDevice, VK_KHR_SAMPLER_YCBCR_CONVERSION_EXTENSION_NAME, VULKAN_EXTENSION_ENABLED, VK_API_VERSION_1_1)
+	{
+	}
+
+	virtual void PrePhysicalDeviceFeatures(VkPhysicalDeviceFeatures2KHR& PhysicalDeviceFeatures2) override final
+	{
+		ZeroVulkanStruct(PhysicalDeviceSamplerYcbcrConversionFeatures, VK_STRUCTURE_TYPE_PHYSICAL_DEVICE_SAMPLER_YCBCR_CONVERSION_FEATURES);
+		AddToPNext(PhysicalDeviceFeatures2, PhysicalDeviceSamplerYcbcrConversionFeatures);
+	}
+
+	virtual void PostPhysicalDeviceFeatures(FOptionalVulkanDeviceExtensions& ExtensionFlags) override final
+	{
+		bRequirementsPassed = (PhysicalDeviceSamplerYcbcrConversionFeatures.samplerYcbcrConversion == VK_TRUE);
+		ExtensionFlags.HasKHRSamplerYcbcrConversion = bRequirementsPassed;
+	}
+
+	virtual void PreCreateDevice(VkDeviceCreateInfo& DeviceCreateInfo) override final
+	{
+		if (bRequirementsPassed)
+		{
+			AddToPNext(DeviceCreateInfo, PhysicalDeviceSamplerYcbcrConversionFeatures);
+		}
+	}
+
+
+private:
+	VkPhysicalDeviceSamplerYcbcrConversionFeatures PhysicalDeviceSamplerYcbcrConversionFeatures;
+};
 
 template <typename ExtensionType>
 static void FlagExtensionSupport(const TArray<VkExtensionProperties>& ExtensionProperties, TArray<TUniquePtr<ExtensionType>>& UEExtensions, uint32 ApiVersion, const TCHAR* ExtensionTypeName)
@@ -1575,6 +1609,7 @@ FVulkanDeviceExtensionArray FVulkanDeviceExtension::GetUESupportedDeviceExtensio
 	ADD_CUSTOM_EXTENSION(FVulkanEXTPipelineCreationCacheControlExtension);
 	ADD_CUSTOM_EXTENSION(FVulkanKHRFragmentShaderBarycentricExtension);
 	ADD_CUSTOM_EXTENSION(FVulkanNVComputeShaderDerivatives);
+	ADD_CUSTOM_EXTENSION(FVulkanKHRSamplerYcbcrConversionExtension);
 
 	// Needed for Raytracing
 	ADD_CUSTOM_EXTENSION(FVulkanKHRBufferDeviceAddressExtension);
diff --git a/Engine/Source/Runtime/VulkanRHI/Private/VulkanLoader.h b/Engine/Source/Runtime/VulkanRHI/Private/VulkanLoader.h
--- a/Engine/Source/Runtime/VulkanRHI/Private/VulkanLoader.h
+++ b/Engine/Source/Runtime/VulkanRHI/Private/VulkanLoader.h
@@ -153,7 +153,9 @@
 	EnumMacro(PFN_vkDestroySwapchainKHR, vkDestroySwapchainKHR) \
 	EnumMacro(PFN_vkGetSwapchainImagesKHR, vkGetSwapchainImagesKHR) \
 	EnumMacro(PFN_vkAcquireNextImageKHR, vkAcquireNextImageKHR) \
-	EnumMacro(PFN_vkQueuePresentKHR, vkQueuePresentKHR)
+	EnumMacro(PFN_vkQueuePresentKHR, vkQueuePresentKHR) \
+	EnumMacro(PFN_vkCreateSamplerYcbcrConversion, vkCreateSamplerYcbcrConversion) \
+	EnumMacro(PFN_vkDestroySamplerYcbcrConversion, vkDestroySamplerYcbcrConversion)
 
 // List all surface Vulkan entry points used by Unreal that need to be loaded manually
 #define ENUM_VK_ENTRYPOINTS_SURFACE_INSTANCE(EnumMacro) \
@@ -222,9 +224,9 @@
 	EnumMacro(PFN_vkDeferredOperationJoinKHR, vkDeferredOperationJoinKHR) \
 	EnumMacro(PFN_vkGetDeferredOperationMaxConcurrencyKHR, vkGetDeferredOperationMaxConcurrencyKHR) \
 	EnumMacro(PFN_vkGetDeferredOperationResultKHR, vkGetDeferredOperationResultKHR) \
-    EnumMacro(PFN_vkCmdWriteBufferMarkerAMD, vkCmdWriteBufferMarkerAMD) \
-    EnumMacro(PFN_vkCmdSetCheckpointNV, vkCmdSetCheckpointNV) \
-    EnumMacro(PFN_vkGetQueueCheckpointDataNV, vkGetQueueCheckpointDataNV) \
+	EnumMacro(PFN_vkCmdWriteBufferMarkerAMD, vkCmdWriteBufferMarkerAMD) \
+	EnumMacro(PFN_vkCmdSetCheckpointNV, vkCmdSetCheckpointNV) \
+	EnumMacro(PFN_vkGetQueueCheckpointDataNV, vkGetQueueCheckpointDataNV) \
 	EnumMacro(PFN_vkGetPhysicalDeviceFragmentShadingRatesKHR, vkGetPhysicalDeviceFragmentShadingRatesKHR) \
 	EnumMacro(PFN_vkGetDeviceFaultInfoEXT, vkGetDeviceFaultInfoEXT)
 
diff --git a/Engine/Source/Runtime/VulkanRHI/Private/VulkanRHI.cpp b/Engine/Source/Runtime/VulkanRHI/Private/VulkanRHI.cpp
--- a/Engine/Source/Runtime/VulkanRHI/Private/VulkanRHI.cpp
+++ b/Engine/Source/Runtime/VulkanRHI/Private/VulkanRHI.cpp
@@ -1621,6 +1621,24 @@ FTextureRHIRef FVulkanDynamicRHI::RHICreateTexture2DFromResource(EPixelFormat Fo
 	return new FVulkanTexture(*Device, Desc, Resource, ExternalImageDeleteCallbackInfo);
 }
 
+#if PLATFORM_ANDROID
+FTextureRHIRef FVulkanDynamicRHI::RHICreateTexture2DFromAndroidHardwareBuffer(AHardwareBuffer* HardwareBuffer)
+{
+	check(HardwareBuffer);
+
+	AHardwareBuffer_Desc HardwareBufferDesc;
+	AHardwareBuffer_describe(HardwareBuffer, &HardwareBufferDesc);
+	check((HardwareBufferDesc.usage & AHARDWAREBUFFER_USAGE_GPU_SAMPLED_IMAGE) != 0);
+
+	const FRHITextureCreateDesc Desc =
+		FRHITextureCreateDesc::Create2D(TEXT("VulkanTexture2DFromAndroidHardwareBuffer"), HardwareBufferDesc.width, HardwareBufferDesc.height, PF_Unknown)
+		.SetFlags(ETextureCreateFlags::External)
+		.DetermineInititialState();
+
+	return new FVulkanTexture(*Device, Desc, HardwareBufferDesc, HardwareBuffer);
+}
+#endif
+
 FTextureRHIRef FVulkanDynamicRHI::RHICreateTexture2DArrayFromResource(EPixelFormat Format, uint32 SizeX, uint32 SizeY, uint32 ArraySize, uint32 NumMips, uint32 NumSamples, VkImage Resource, ETextureCreateFlags Flags, const FClearValueBinding& ClearValueBinding)
 {
 	const FRHITextureCreateDesc Desc =
diff --git a/Engine/Source/Runtime/VulkanRHI/Private/VulkanResources.h b/Engine/Source/Runtime/VulkanRHI/Private/VulkanResources.h
--- a/Engine/Source/Runtime/VulkanRHI/Private/VulkanResources.h
+++ b/Engine/Source/Runtime/VulkanRHI/Private/VulkanResources.h
@@ -475,7 +475,8 @@ public:
 		, uint32 ArraySliceIndex
 		, uint32 NumArraySlices
 		, bool bUseIdentitySwizzle = false
-		, VkImageUsageFlags ImageUsageFlags = 0);
+		, VkImageUsageFlags ImageUsageFlags = 0
+		, VkSamplerYcbcrConversion SamplerYcbcrConversion = nullptr);
 
 	// NOTE: The InOffset applies to the FVulkanResourceMultiBuffer (it does not include any internal Allocation offsets that may exist)
 	FVulkanView* InitAsStructuredBufferView(
@@ -563,9 +564,13 @@ public:
 	{}
 
 	// Construct from external resource.
-	// FIXME: HUGE HACK: the bUnused argument is there to disambiguate this overload from the one above when passing nullptr, since nullptr is a valid VkImage. Get rid of this code smell when unifying FVulkanSurface and FVulkanTexture.
 	FVulkanTexture(FVulkanDevice& InDevice, const FRHITextureCreateDesc& InCreateDesc, VkImage InImage, const FVulkanRHIExternalImageDeleteCallbackInfo& InExternalImageDeleteCallbackInfo);
 
+#if PLATFORM_ANDROID
+	// Construct from Android Hardware buffer. HardwareBuffer_Desc could be extracted from HardwareBuffer, but adding it here makes the function signature unambigious
+	FVulkanTexture(FVulkanDevice& InDevice, const FRHITextureCreateDesc& InCreateDesc, const AHardwareBuffer_Desc& HardwareBufferDesc, AHardwareBuffer* HardwareBuffer);
+#endif
+
 	// Aliasing constructor.
 	FVulkanTexture(FVulkanDevice& InDevice, const FRHITextureCreateDesc& InCreateDesc, FTextureRHIRef& SrcTextureRHI);
 
diff --git a/Engine/Source/Runtime/VulkanRHI/Private/VulkanTexture.cpp b/Engine/Source/Runtime/VulkanRHI/Private/VulkanTexture.cpp
--- a/Engine/Source/Runtime/VulkanRHI/Private/VulkanTexture.cpp
+++ b/Engine/Source/Runtime/VulkanRHI/Private/VulkanTexture.cpp
@@ -1871,6 +1871,268 @@ FVulkanTexture::FVulkanTexture(FVulkanDevice& InDevice, const FRHITextureCreateD
 	}
 }
 
+#if PLATFORM_ANDROID
+
+struct FVulkanAndroidTextureResources
+{
+	VkImage Image;
+	VkDeviceMemory DeviceMemory;
+	VkSamplerYcbcrConversion SamplerYcbcrConversion;
+	AHardwareBuffer* HardwareBuffer;
+};
+
+static void CleanupVulkanAndroidTextureResources(void* UserData)
+{
+	check(UserData);
+
+	FVulkanAndroidTextureResources* VulkanResources = static_cast<FVulkanAndroidTextureResources*>(UserData);
+
+	IVulkanDynamicRHI* RHI = GetIVulkanDynamicRHI();
+	VkDevice Device = RHI->RHIGetVkDevice();
+	const VkAllocationCallbacks* AllocationCallbacks = RHI->RHIGetVkAllocationCallbacks();
+
+	if (VulkanResources->SamplerYcbcrConversion != VK_NULL_HANDLE)
+	{
+		VulkanRHI::vkDestroySamplerYcbcrConversion(Device, VulkanResources->SamplerYcbcrConversion, AllocationCallbacks);
+	}
+
+	if (VulkanResources->DeviceMemory != VK_NULL_HANDLE)
+	{
+		VulkanRHI::vkFreeMemory(Device, VulkanResources->DeviceMemory, AllocationCallbacks);
+	}
+
+	if (VulkanResources->Image != VK_NULL_HANDLE)
+	{
+		VulkanRHI::vkDestroyImage(Device, VulkanResources->Image, AllocationCallbacks);
+	}
+
+	if (VulkanResources->HardwareBuffer)
+	{
+		AHardwareBuffer_release(VulkanResources->HardwareBuffer);
+	}
+
+	delete VulkanResources;
+}
+
+FVulkanTexture::FVulkanTexture(FVulkanDevice& InDevice, const FRHITextureCreateDesc& InCreateDesc, const AHardwareBuffer_Desc& HardwareBufferDesc, AHardwareBuffer* HardwareBuffer)
+	: FRHITexture(InCreateDesc)
+	, Device(&InDevice)
+	, ImageUsageFlags(0)
+	, StorageFormat(VK_FORMAT_UNDEFINED)
+	, ViewFormat(VK_FORMAT_UNDEFINED)
+	, MemProps(VK_MEMORY_PROPERTY_DEVICE_LOCAL_BIT)
+	, Tiling(VK_IMAGE_TILING_MAX_ENUM)	// Can be expanded to a per-platform definition
+	, FullAspectMask(0)
+	, PartialAspectMask(0)
+	, CpuReadbackBuffer(nullptr)
+	, DefaultLayout(VK_IMAGE_LAYOUT_UNDEFINED)
+	, ImageOwnerType(EImageOwnerType::ExternalOwner)
+{
+	VULKAN_TRACK_OBJECT_CREATE(FVulkanTexture, this);
+
+	check(HardwareBuffer);
+	AHardwareBuffer_acquire(HardwareBuffer);
+
+	IVulkanDynamicRHI* RHI = GetIVulkanDynamicRHI();
+	VkDevice VulkanDevice = InDevice.GetInstanceHandle();
+	const VkAllocationCallbacks* AllocationCallbacks = RHI->RHIGetVkAllocationCallbacks();
+
+	VkAndroidHardwareBufferFormatPropertiesANDROID HardwareBufferFormatProperties;
+	ZeroVulkanStruct(HardwareBufferFormatProperties, VK_STRUCTURE_TYPE_ANDROID_HARDWARE_BUFFER_FORMAT_PROPERTIES_ANDROID);
+
+	VkAndroidHardwareBufferPropertiesANDROID HardwareBufferProperties;
+	ZeroVulkanStruct(HardwareBufferProperties, VK_STRUCTURE_TYPE_ANDROID_HARDWARE_BUFFER_PROPERTIES_ANDROID);
+	HardwareBufferProperties.pNext = &HardwareBufferFormatProperties;
+
+	VERIFYVULKANRESULT(VulkanRHI::vkGetAndroidHardwareBufferPropertiesANDROID(VulkanDevice, HardwareBuffer, &HardwareBufferProperties));
+
+	VkExternalFormatANDROID ExternalFormat;
+	ZeroVulkanStruct(ExternalFormat, VK_STRUCTURE_TYPE_EXTERNAL_FORMAT_ANDROID);
+	ExternalFormat.externalFormat = HardwareBufferFormatProperties.externalFormat;
+
+	VkExternalMemoryImageCreateInfo ExternalMemoryImageCreateInfo;
+	ZeroVulkanStruct(ExternalMemoryImageCreateInfo, VK_STRUCTURE_TYPE_EXTERNAL_MEMORY_IMAGE_CREATE_INFO);
+	ExternalMemoryImageCreateInfo.handleTypes = VK_EXTERNAL_MEMORY_HANDLE_TYPE_ANDROID_HARDWARE_BUFFER_BIT_ANDROID;
+	ExternalMemoryImageCreateInfo.pNext = &ExternalFormat;
+
+	VkImageCreateInfo ImageCreateInfo;
+	ZeroVulkanStruct(ImageCreateInfo, VK_STRUCTURE_TYPE_IMAGE_CREATE_INFO);
+	ImageCreateInfo.imageType = VK_IMAGE_TYPE_2D;
+	ImageCreateInfo.format = VK_FORMAT_UNDEFINED;
+	ImageCreateInfo.extent.width = HardwareBufferDesc.width;
+	ImageCreateInfo.extent.height = HardwareBufferDesc.height;
+	ImageCreateInfo.extent.depth = 1;
+	ImageCreateInfo.mipLevels = 1;
+	ImageCreateInfo.arrayLayers = HardwareBufferDesc.layers;
+
+	ImageCreateInfo.flags = 0;
+	ImageCreateInfo.samples = VK_SAMPLE_COUNT_1_BIT;
+
+	ImageCreateInfo.tiling = VK_IMAGE_TILING_OPTIMAL;
+	ImageCreateInfo.usage = VK_IMAGE_USAGE_SAMPLED_BIT;
+	ImageCreateInfo.sharingMode = VK_SHARING_MODE_EXCLUSIVE;
+	ImageCreateInfo.queueFamilyIndexCount = 0;
+	ImageCreateInfo.pQueueFamilyIndices = nullptr;
+	ImageCreateInfo.initialLayout = VK_IMAGE_LAYOUT_UNDEFINED;
+
+	ImageCreateInfo.pNext = &ExternalMemoryImageCreateInfo;
+
+	VkImage VulkanImage;
+	VERIFYVULKANRESULT(VulkanRHI::vkCreateImage(VulkanDevice, &ImageCreateInfo, nullptr, &VulkanImage));
+
+	VkMemoryDedicatedAllocateInfo MemoryDedicatedAllocateInfo;
+	ZeroVulkanStruct(MemoryDedicatedAllocateInfo, VK_STRUCTURE_TYPE_MEMORY_DEDICATED_ALLOCATE_INFO);
+	MemoryDedicatedAllocateInfo.image = VulkanImage;
+	MemoryDedicatedAllocateInfo.buffer = VK_NULL_HANDLE;
+
+	VkImportAndroidHardwareBufferInfoANDROID ImportAndroidHardwareBufferInfo;
+	ZeroVulkanStruct(ImportAndroidHardwareBufferInfo, VK_STRUCTURE_TYPE_IMPORT_ANDROID_HARDWARE_BUFFER_INFO_ANDROID);
+	ImportAndroidHardwareBufferInfo.buffer = HardwareBuffer;
+	ImportAndroidHardwareBufferInfo.pNext = &MemoryDedicatedAllocateInfo;
+
+	uint32 MemoryTypeBits = HardwareBufferProperties.memoryTypeBits;
+	check(MemoryTypeBits > 0); // No index available, this should never happen
+	uint32 MemoryTypeIndex = 0;
+	for (;(MemoryTypeBits & 1) != 1; ++MemoryTypeIndex)
+	{
+		MemoryTypeBits >>= 1;
+	}
+
+	VkMemoryAllocateInfo MemoryAllocateInfo;
+	ZeroVulkanStruct(MemoryAllocateInfo, VK_STRUCTURE_TYPE_MEMORY_ALLOCATE_INFO);
+	MemoryAllocateInfo.allocationSize = HardwareBufferProperties.allocationSize;
+	MemoryAllocateInfo.memoryTypeIndex = MemoryTypeIndex;
+	MemoryAllocateInfo.pNext = &ImportAndroidHardwareBufferInfo;
+
+	VkDeviceMemory VulkanDeviceMemory;
+	VERIFYVULKANRESULT(VulkanRHI::vkAllocateMemory(VulkanDevice, &MemoryAllocateInfo, AllocationCallbacks, &VulkanDeviceMemory));
+	VERIFYVULKANRESULT(VulkanRHI::vkBindImageMemory(VulkanDevice, VulkanImage, VulkanDeviceMemory, 0));
+
+	VkSamplerYcbcrConversionCreateInfo SamplerYcbcrConversionCreateInfo;
+	ZeroVulkanStruct(SamplerYcbcrConversionCreateInfo, VK_STRUCTURE_TYPE_SAMPLER_YCBCR_CONVERSION_CREATE_INFO);
+	SamplerYcbcrConversionCreateInfo.format = VK_FORMAT_UNDEFINED;
+	SamplerYcbcrConversionCreateInfo.ycbcrModel = HardwareBufferFormatProperties.suggestedYcbcrModel;
+	SamplerYcbcrConversionCreateInfo.ycbcrRange = HardwareBufferFormatProperties.suggestedYcbcrRange;
+	SamplerYcbcrConversionCreateInfo.components = HardwareBufferFormatProperties.samplerYcbcrConversionComponents;
+	SamplerYcbcrConversionCreateInfo.xChromaOffset = HardwareBufferFormatProperties.suggestedXChromaOffset;
+	SamplerYcbcrConversionCreateInfo.yChromaOffset = HardwareBufferFormatProperties.suggestedYChromaOffset;
+	SamplerYcbcrConversionCreateInfo.chromaFilter = VK_FILTER_LINEAR;
+	SamplerYcbcrConversionCreateInfo.forceExplicitReconstruction = VK_FALSE;
+	SamplerYcbcrConversionCreateInfo.pNext = &ExternalFormat;
+
+	VkSamplerYcbcrConversion SamplerYcbcrConversion;
+	VERIFYVULKANRESULT(VulkanRHI::vkCreateSamplerYcbcrConversion(VulkanDevice, &SamplerYcbcrConversionCreateInfo, AllocationCallbacks, &SamplerYcbcrConversion));
+
+	FVulkanAndroidTextureResources* VulkanAndroidTextureResources = new FVulkanAndroidTextureResources{ VulkanImage, VulkanDeviceMemory, SamplerYcbcrConversion };
+	ExternalImageDeleteCallbackInfo =
+	{
+		VulkanAndroidTextureResources,
+		CleanupVulkanAndroidTextureResources
+	};
+
+	Image = VulkanImage;
+
+	// From here this is the same as the ctor that takes an VkImage, excepct for passing the SamplerYcbcrConversion to the view, 
+	// possibly some code could be shared.
+	{
+		StorageFormat = UEToVkTextureFormat(InCreateDesc.Format, false);
+
+		checkf(InCreateDesc.Format == PF_Unknown || StorageFormat != VK_FORMAT_UNDEFINED, TEXT("PixelFormat %d, is not supported for images"), (int32)InCreateDesc.Format);
+
+		ViewFormat = UEToVkTextureFormat(InCreateDesc.Format, EnumHasAllFlags(InCreateDesc.Flags, TexCreate_SRGB));
+		FullAspectMask = VulkanRHI::GetAspectMaskFromUEFormat(InCreateDesc.Format, true, true);
+		PartialAspectMask = VulkanRHI::GetAspectMaskFromUEFormat(InCreateDesc.Format, false, true);
+
+		// Purely informative patching, we know that "TexCreate_Presentable" uses optimal tiling
+		if (EnumHasAllFlags(InCreateDesc.Flags, TexCreate_Presentable) && GetTiling() == VK_IMAGE_TILING_MAX_ENUM)
+		{
+			Tiling = VK_IMAGE_TILING_OPTIMAL;
+		}
+
+		if (Image != VK_NULL_HANDLE)
+		{
+			ImageUsageFlags = GetUsageFlagsFromCreateFlags(InDevice, InCreateDesc.Flags);
+#if VULKAN_ENABLE_WRAP_LAYER
+			FWrapLayer::CreateImage(VK_SUCCESS, InDevice.GetInstanceHandle(), nullptr, &Image);
+#endif
+			VULKAN_SET_DEBUG_NAME(InDevice, VK_OBJECT_TYPE_IMAGE, Image, TEXT("%s:(FVulkanTexture*)0x%p"), InCreateDesc.DebugName ? InCreateDesc.DebugName : TEXT("?"), this);
+
+			const bool bRenderTarget = EnumHasAnyFlags(InCreateDesc.Flags, TexCreate_RenderTargetable | TexCreate_DepthStencilTargetable);
+			const VkImageLayout InitialLayout = GetInitialLayoutFromRHIAccess(InCreateDesc.InitialState, bRenderTarget && IsDepthOrStencilAspect(), SupportsSampling());
+			const bool bDoInitialClear = bRenderTarget;
+			const bool bOnlyAddToLayoutManager = !bRenderTarget;
+
+			DefaultLayout = InitialLayout;
+
+			FRHICommandList& RHICmdList = FRHICommandListExecutor::GetImmediateCommandList();
+			if (!IsInRenderingThread() || (RHICmdList.Bypass() || !IsRunningRHIInSeparateThread()))
+			{
+				FVulkanCommandListContext& Context = InDevice.GetImmediateContext();
+				if (bOnlyAddToLayoutManager)
+				{
+					FVulkanCmdBuffer* CmdBuffer = Context.GetCommandBufferManager()->GetActiveCmdBuffer();
+					CmdBuffer->GetLayoutManager().SetFullLayout(*this, InitialLayout, true);
+				}
+				else if (InitialLayout != VK_IMAGE_LAYOUT_UNDEFINED || bDoInitialClear)
+				{
+					SetInitialImageState(Context, InitialLayout, bDoInitialClear, InCreateDesc.ClearValue, false);
+				}
+			}
+			else
+			{
+				check(IsInRenderingThread());
+				ALLOC_COMMAND_CL(RHICmdList, FRHICommandSetInitialImageState)(this, InitialLayout, bOnlyAddToLayoutManager, bDoInitialClear, InCreateDesc.ClearValue, false);
+			}
+		}
+	}
+
+	const VkImageViewType ViewType = GetViewType();
+	const VkDescriptorType DescriptorType = SupportsSampling() ? VK_DESCRIPTOR_TYPE_SAMPLED_IMAGE : VK_DESCRIPTOR_TYPE_STORAGE_IMAGE;
+	const bool bUseIdentitySwizzle = (DescriptorType != VK_DESCRIPTOR_TYPE_SAMPLED_IMAGE) ||
+									(ViewFormat == VK_FORMAT_UNDEFINED); // External buffer textures also require identity swizzle
+
+	if (Image != VK_NULL_HANDLE)
+	{
+		DefaultView = (new FVulkanView(InDevice, DescriptorType))->InitAsTextureView(
+			Image
+			, ViewType
+			, GetFullAspectMask()
+			, InCreateDesc.Format
+			, ViewFormat
+			, 0
+			, FMath::Max(InCreateDesc.NumMips, (uint8)1u)
+			, 0
+			, GetNumberOfArrayLevels()
+			, bUseIdentitySwizzle
+			, 0
+			, SamplerYcbcrConversion
+		);
+	}
+
+	if (FullAspectMask == PartialAspectMask)
+	{
+		PartialView = DefaultView;
+	}
+	else
+	{
+		PartialView = (new FVulkanView(InDevice, DescriptorType))->InitAsTextureView(
+			Image
+			, ViewType
+			, PartialAspectMask
+			, InCreateDesc.Format
+			, ViewFormat
+			, 0
+			, FMath::Max(InCreateDesc.NumMips, (uint8)1u)
+			, 0
+			, GetNumberOfArrayLevels()
+			, false
+		);
+	}
+}
+
+#endif // PLATFORM_ANDROID
+
+
 FVulkanTexture::FVulkanTexture(FVulkanDevice& InDevice, const FRHITextureCreateDesc& InCreateDesc, FTextureRHIRef& SrcTextureRHI)
 	: FRHITexture(InCreateDesc)
 	, Device(&InDevice)
diff --git a/Engine/Source/Runtime/VulkanRHI/Private/VulkanUAV.cpp b/Engine/Source/Runtime/VulkanRHI/Private/VulkanUAV.cpp
--- a/Engine/Source/Runtime/VulkanRHI/Private/VulkanUAV.cpp
+++ b/Engine/Source/Runtime/VulkanRHI/Private/VulkanUAV.cpp
@@ -122,7 +122,8 @@ FVulkanView* FVulkanView::InitAsTextureView(
 	, uint32 ArraySliceIndex
 	, uint32 NumArraySlices
 	, bool bUseIdentitySwizzle
-	, VkImageUsageFlags ImageUsageFlags)
+	, VkImageUsageFlags ImageUsageFlags
+	, VkSamplerYcbcrConversion SamplerYcbcrConversion)
 {
 	// We will need a deferred update if the descriptor was already in use
 	const bool bImmediateUpdate = !IsInitialized();
@@ -150,6 +151,15 @@ FVulkanView* FVulkanView::InitAsTextureView(
 	}
 #endif
 
+	VkSamplerYcbcrConversionInfo SamplerYcbcrConversionInfo;
+	if (SamplerYcbcrConversion)
+	{
+		ZeroVulkanStruct(SamplerYcbcrConversionInfo, VK_STRUCTURE_TYPE_SAMPLER_YCBCR_CONVERSION_INFO);
+		SamplerYcbcrConversionInfo.conversion = SamplerYcbcrConversion;
+		SamplerYcbcrConversionInfo.pNext = ViewInfo.pNext;
+		ViewInfo.pNext = &SamplerYcbcrConversionInfo;
+	}
+
 	if (bUseIdentitySwizzle)
 	{
 		ViewInfo.components.a = VK_COMPONENT_SWIZZLE_IDENTITY;
diff --git a/Engine/Source/Runtime/VulkanRHI/Public/IVulkanDynamicRHI.h b/Engine/Source/Runtime/VulkanRHI/Public/IVulkanDynamicRHI.h
--- a/Engine/Source/Runtime/VulkanRHI/Public/IVulkanDynamicRHI.h
+++ b/Engine/Source/Runtime/VulkanRHI/Public/IVulkanDynamicRHI.h
@@ -59,6 +59,9 @@ struct IVulkanDynamicRHI : public FDynamicRHI
 	virtual TArray<FAnsiString> RHIGetLoadedDeviceExtensions() const = 0;
 
 	virtual FTextureRHIRef   RHICreateTexture2DFromResource(EPixelFormat Format, uint32 SizeX, uint32 SizeY, uint32 NumMips, uint32 NumSamples, VkImage Resource, ETextureCreateFlags Flags, const FClearValueBinding& ClearValueBinding = FClearValueBinding::Transparent, const FVulkanRHIExternalImageDeleteCallbackInfo& ExternalImageDeleteCallbackInfo = {}) = 0;
+#if PLATFORM_ANDROID
+	virtual FTextureRHIRef   RHICreateTexture2DFromAndroidHardwareBuffer(AHardwareBuffer* HardwareBuffer) = 0;
+#endif
 	virtual FTextureRHIRef   RHICreateTexture2DArrayFromResource(EPixelFormat Format, uint32 SizeX, uint32 SizeY, uint32 ArraySize, uint32 NumMips, uint32 NumSamples, VkImage Resource, ETextureCreateFlags Flags, const FClearValueBinding& ClearValueBinding = FClearValueBinding::Transparent) = 0;
 	virtual FTextureRHIRef   RHICreateTextureCubeFromResource(EPixelFormat Format, uint32 Size, bool bArray, uint32 ArraySize, uint32 NumMips, VkImage Resource, ETextureCreateFlags Flags, const FClearValueBinding& ClearValueBinding = FClearValueBinding::Transparent) = 0;
 
diff --git a/Engine/Source/Runtime/VulkanRHI/Public/VulkanDynamicRHI.h b/Engine/Source/Runtime/VulkanRHI/Public/VulkanDynamicRHI.h
--- a/Engine/Source/Runtime/VulkanRHI/Public/VulkanDynamicRHI.h
+++ b/Engine/Source/Runtime/VulkanRHI/Public/VulkanDynamicRHI.h
@@ -75,6 +75,9 @@ public:
 	virtual TArray<VkExtensionProperties> RHIGetAllDeviceExtensions(VkPhysicalDevice InPhysicalDevice) const final override;
 	virtual TArray<FAnsiString> RHIGetLoadedDeviceExtensions() const final override;
 	virtual FTextureRHIRef RHICreateTexture2DFromResource(EPixelFormat Format, uint32 SizeX, uint32 SizeY, uint32 NumMips, uint32 NumSamples, VkImage Resource, ETextureCreateFlags Flags, const FClearValueBinding& ClearValueBinding = FClearValueBinding::Transparent, const FVulkanRHIExternalImageDeleteCallbackInfo& ExternalImageDeleteCallbackInfo = {}) final override;
+#if PLATFORM_ANDROID
+	virtual FTextureRHIRef RHICreateTexture2DFromAndroidHardwareBuffer(AHardwareBuffer* HardwareBuffer) final override;
+#endif
 	virtual FTextureRHIRef RHICreateTexture2DArrayFromResource(EPixelFormat Format, uint32 SizeX, uint32 SizeY, uint32 ArraySize, uint32 NumMips, uint32 NumSamples, VkImage Resource, ETextureCreateFlags Flags, const FClearValueBinding& ClearValueBinding = FClearValueBinding::Transparent) final override;
 	virtual FTextureRHIRef RHICreateTextureCubeFromResource(EPixelFormat Format, uint32 Size, bool bArray, uint32 ArraySize, uint32 NumMips, VkImage Resource, ETextureCreateFlags Flags, const FClearValueBinding& ClearValueBinding = FClearValueBinding::Transparent) final override;
 	virtual VkImage RHIGetVkImage(FRHITexture* InTexture) const final override;
diff --git a/Engine/Source/Runtime/VulkanRHI/Public/VulkanThirdParty.h b/Engine/Source/Runtime/VulkanRHI/Public/VulkanThirdParty.h
--- a/Engine/Source/Runtime/VulkanRHI/Public/VulkanThirdParty.h
+++ b/Engine/Source/Runtime/VulkanRHI/Public/VulkanThirdParty.h
@@ -54,3 +54,7 @@ THIRD_PARTY_INCLUDES_END
 #endif
 
 #endif
+
+#if PLATFORM_ANDROID
+#include <android/hardware_buffer.h>
+#endif
\ No newline at end of file
