Added missing TEXT wrapper

#rnx

diff --git a/Engine/Plugins/Media/ElectraUtil/Source/ElectraSamples/Private/ElectraTextureSampleAndroid.cpp b/Engine/Plugins/Media/ElectraUtil/Source/ElectraSamples/Private/ElectraTextureSampleAndroid.cpp
--- a/Engine/Plugins/Media/ElectraUtil/Source/ElectraSamples/Private/ElectraTextureSampleAndroid.cpp
+++ b/Engine/Plugins/Media/ElectraUtil/Source/ElectraSamples/Private/ElectraTextureSampleAndroid.cpp
@@ -237,11 +237,11 @@ FElectraTextureSampleSupport::FElectraTextureSampleSupport()
 		{
 			IVulkanDynamicRHI* RHI = GetIVulkanDynamicRHI();
 			TArray<FAnsiString> LoadedDeviceExtensions = RHI->RHIGetLoadedDeviceExtensions();
-			if (!LoadedDeviceExtensions.Find(VK_KHR_SAMPLER_YCBCR_CONVERSION_EXTENSION_NAME))
+			if (!LoadedDeviceExtensions.Find(TEXT(VK_KHR_SAMPLER_YCBCR_CONVERSION_EXTENSION_NAME)))
 			{
 				UE_LOG(LogElectraSamples, Log, TEXT("Selecting CPU path because GPU extension '" VK_KHR_SAMPLER_YCBCR_CONVERSION_EXTENSION_NAME "' is not available!"));
 			}
-			else if (!LoadedDeviceExtensions.Find(VK_ANDROID_EXTERNAL_MEMORY_ANDROID_HARDWARE_BUFFER_EXTENSION_NAME))
+			else if (!LoadedDeviceExtensions.Find(TEXT(VK_ANDROID_EXTERNAL_MEMORY_ANDROID_HARDWARE_BUFFER_EXTENSION_NAME)))
 			{
 				UE_LOG(LogElectraSamples, Log, TEXT("Selecting CPU path because GPU extension '" VK_ANDROID_EXTERNAL_MEMORY_ANDROID_HARDWARE_BUFFER_EXTENSION_NAME "' is not available!"));
 			}
