Another attempt to fix Horde compile issue...

#rnx

diff --git a/Engine/Plugins/Media/ElectraUtil/Source/ElectraSamples/Private/ElectraTextureSampleAndroid.cpp b/Engine/Plugins/Media/ElectraUtil/Source/ElectraSamples/Private/ElectraTextureSampleAndroid.cpp
--- a/Engine/Plugins/Media/ElectraUtil/Source/ElectraSamples/Private/ElectraTextureSampleAndroid.cpp
+++ b/Engine/Plugins/Media/ElectraUtil/Source/ElectraSamples/Private/ElectraTextureSampleAndroid.cpp
@@ -11,24 +11,24 @@
 #include "Android/AndroidJava.h"
 #include "Android/AndroidApplication.h"
 
+#include "Containers/AnsiString.h"
+
 #include "RenderingThread.h"
 #include "RHIStaticStates.h"
 #include "PipelineStateCache.h"
 
 #include "RenderUtils.h"
 
-#include <EGL/egl.h>
-#include <EGL/eglext.h>
-#include <GLES3/gl31.h>
-#include "GLES2/gl2ext.h"
-#include "android/hardware_buffer.h"
-#include "android/hardware_buffer_jni.h"
-
 #include "VulkanCommon.h"
 #include "IVulkanDynamicRHI.h"
 
-#include "vulkan/vulkan_android.h"
-
+#include <EGL/egl.h>
+#include <EGL/eglext.h>
+#include <GLES3/gl31.h>
+#include <GLES2/gl2ext.h>
+#include <android/hardware_buffer.h>
+#include <android/hardware_buffer_jni.h>
+#include <vulkan/vulkan_android.h>
 
 DECLARE_GPU_STAT_NAMED(MediaAndroidDecoder_Convert, TEXT("MediaAndroidDecoder_Convert"));
 
@@ -237,11 +237,11 @@ FElectraTextureSampleSupport::FElectraTextureSampleSupport()
 		{
 			IVulkanDynamicRHI* RHI = GetIVulkanDynamicRHI();
 			TArray<FAnsiString> LoadedDeviceExtensions = RHI->RHIGetLoadedDeviceExtensions();
-			if (!LoadedDeviceExtensions.Find(TEXT(VK_KHR_SAMPLER_YCBCR_CONVERSION_EXTENSION_NAME)))
+			if (!LoadedDeviceExtensions.Find(FAnsiString(VK_KHR_SAMPLER_YCBCR_CONVERSION_EXTENSION_NAME)))
 			{
 				UE_LOG(LogElectraSamples, Log, TEXT("Selecting CPU path because GPU extension '" VK_KHR_SAMPLER_YCBCR_CONVERSION_EXTENSION_NAME "' is not available!"));
 			}
-			else if (!LoadedDeviceExtensions.Find(TEXT(VK_ANDROID_EXTERNAL_MEMORY_ANDROID_HARDWARE_BUFFER_EXTENSION_NAME)))
+			else if (!LoadedDeviceExtensions.Find(FAnsiString(VK_ANDROID_EXTERNAL_MEMORY_ANDROID_HARDWARE_BUFFER_EXTENSION_NAME)))
 			{
 				UE_LOG(LogElectraSamples, Log, TEXT("Selecting CPU path because GPU extension '" VK_ANDROID_EXTERNAL_MEMORY_ANDROID_HARDWARE_BUFFER_EXTENSION_NAME "' is not available!"));
 			}
