Fixed race condition where the VideoDecoderOutput would be shutdown before the copy render command was executed, leading to an infinite stall
#rnx

diff --git a/Engine/Plugins/Media/ElectraUtil/Source/ElectraSamples/Private/ElectraTextureSamplePC.cpp b/Engine/Plugins/Media/ElectraUtil/Source/ElectraSamples/Private/ElectraTextureSamplePC.cpp
--- a/Engine/Plugins/Media/ElectraUtil/Source/ElectraSamples/Private/ElectraTextureSamplePC.cpp
+++ b/Engine/Plugins/Media/ElectraUtil/Source/ElectraSamples/Private/ElectraTextureSamplePC.cpp
@@ -322,13 +322,8 @@ struct FRHICommandCopyResourceDX11 final : public FRHICommand<FRHICommandCopyRes
 										}
 										else if (Result == ((HRESULT)WAIT_TIMEOUT))
 										{
-											// If key 1 cannot be acquired, another reader is already copying the resource
-											// and will release key with 2. 
-											// Wait to acquire key 2.
-											if (KeyedMutex->AcquireSync(2, INFINITE) == S_OK)
-											{
-												KeyedMutex->ReleaseSync(2);
-											}
+											// If key 1 cannot be acquired, the resource has already been shutdown or consumed by another reader
+											UE_LOG(LogElectraSamples, Warning, TEXT("AcquireSync timed out, DecoderOutput has likely been shut down already!"));
 										}
 										else
 										{
