Electra: Updating helpers for upcoming new features

#rnx

diff --git a/Engine/Plugins/Media/ElectraCodecs/Source/ElectraDecoders/Private/Utils/MPEG/ElectraUtilsMPEGVideo.cpp b/Engine/Plugins/Media/ElectraCodecs/Source/ElectraDecoders/Private/Utils/MPEG/ElectraUtilsMPEGVideo.cpp
--- a/Engine/Plugins/Media/ElectraCodecs/Source/ElectraDecoders/Private/Utils/MPEG/ElectraUtilsMPEGVideo.cpp
+++ b/Engine/Plugins/Media/ElectraCodecs/Source/ElectraDecoders/Private/Utils/MPEG/ElectraUtilsMPEGVideo.cpp
@@ -4,6 +4,7 @@
 #include "Utils/ElectraBitstreamReader.h"
 #include "Utils/MPEG/ElectraUtilsMPEGVideo_H264.h"
 #include "Utils/MPEG/ElectraUtilsMPEGVideo_H265.h"
+#include "Utils/MPEG/ElectraUtilsMP4.h"
 #include "ElectraDecodersUtils.h"
 
 namespace ElectraDecodersUtil
@@ -1361,6 +1362,32 @@ namespace ElectraDecodersUtil
 			return false;
 		}
 
+		bool ParseFromMDCVBox(FSEImastering_display_colour_volume& OutMDCV, const TConstArrayView<uint8>& InMDCVBox)
+		{
+			if (InMDCVBox.IsEmpty())
+			{
+				return false;
+			}
+			FMP4AtomReader boxReader(InMDCVBox.GetData(), InMDCVBox.Num());
+			for(int32 i=0; i<3; ++i)
+			{
+				if (!boxReader.Read(OutMDCV.display_primaries_x[i]) || !boxReader.Read(OutMDCV.display_primaries_y[i]))
+				{
+					return false;
+				}
+			}
+			if (!boxReader.Read(OutMDCV.white_point_x) || !boxReader.Read(OutMDCV.white_point_y))
+			{
+				return false;
+			}
+			if (!boxReader.Read(OutMDCV.max_display_mastering_luminance) || !boxReader.Read(OutMDCV.min_display_mastering_luminance))
+			{
+				return false;
+			}
+			return true;
+		}
+
+
 		bool ParseFromSEIMessage(FSEIcontent_light_level_info& OutCLLI, const FSEIMessage& InMessage)
 		{
 			if (InMessage.PayloadType == FSEIMessage::EPayloadType::PT_content_light_level_info && InMessage.Message.Num() >= 4)
@@ -1372,6 +1399,38 @@ namespace ElectraDecodersUtil
 			}
 			return false;
 		}
+		bool ParseFromCOLLBox(FSEIcontent_light_level_info& OutCLLI, const TConstArrayView<uint8>& InCOLLBox)
+		{
+			if (InCOLLBox.Num() > 4)
+			{
+				uint8 BoxVersion = InCOLLBox[0];
+				if (BoxVersion == 0)
+				{
+					// 'clli' box is the same as a version 0 'COLL 'coll' box.
+					return ParseFromCLLIBox(OutCLLI, MakeConstArrayView(InCOLLBox.GetData() + 4, InCOLLBox.Num() - 4));
+				}
+			}
+			return false;
+		}
+		bool ParseFromCLLIBox(FSEIcontent_light_level_info& OutCLLI, const TConstArrayView<uint8>& InCLLIBox)
+		{
+			static_assert(sizeof(FSEIcontent_light_level_info::max_content_light_level) == 2);
+			static_assert(sizeof(FSEIcontent_light_level_info::max_pic_average_light_level) == 2);
+			if (InCLLIBox.IsEmpty())
+			{
+				return false;
+			}
+			ElectraDecodersUtil::FMP4AtomReader boxReader(InCLLIBox.GetData(), InCLLIBox.Num());
+			OutCLLI.max_content_light_level = 0;			// MaxCLL
+			OutCLLI.max_pic_average_light_level = 0;		// MaxFALL
+			if (!boxReader.Read(OutCLLI.max_content_light_level) || !boxReader.Read(OutCLLI.max_pic_average_light_level))
+			{
+				return false;
+			}
+			return true;
+		}
+
+
 
 		bool ParseFromSEIMessage(FSEIalternative_transfer_characteristics& OutATC, const FSEIMessage& InMessage)
 		{
@@ -1384,5 +1443,37 @@ namespace ElectraDecodersUtil
 			return false;
 		}
 
+		bool ParseFromCOLRBox(FCommonColorimetry& OutColorimetry, const TConstArrayView<uint8>& InCOLRBox)
+		{
+			if (InCOLRBox.Num())
+			{
+				FMP4AtomReader boxReader(InCOLRBox.GetData(), InCOLRBox.Num());
+				uint32 Type;
+				if (boxReader.Read(Type))
+				{
+					if (Type == ElectraDecodersUtil::MakeMP4Atom('n','c','l','x') || Type == ElectraDecodersUtil::MakeMP4Atom('n', 'c', 'l', 'c'))
+					{
+						uint16 colour_primaries, transfer_characteristics, matrix_coeffs;
+						uint8 video_full_range_flag = 0;
+						if (boxReader.Read(colour_primaries) && boxReader.Read(transfer_characteristics) && boxReader.Read(matrix_coeffs))
+						{
+							if (Type == ElectraDecodersUtil::MakeMP4Atom('n','c','l','x') && !boxReader.Read(video_full_range_flag))
+							{
+								return false;
+							}
+							OutColorimetry.colour_primaries = (uint8) colour_primaries;
+							OutColorimetry.transfer_characteristics = (uint8)transfer_characteristics;
+							OutColorimetry.matrix_coeffs = (uint8)matrix_coeffs;
+							OutColorimetry.video_full_range_flag = (uint8)video_full_range_flag;
+							OutColorimetry.video_format = 5;	// Unspecified video format
+							return true;
+						}
+					}
+				}
+			}
+			return false;
+		}
+
+
 	} // namespace MPEG
 } // namespace ElectraDecodersUtil
diff --git a/Engine/Plugins/Media/ElectraCodecs/Source/ElectraDecoders/Public/ElectraDecodersUtils.h b/Engine/Plugins/Media/ElectraCodecs/Source/ElectraDecoders/Public/ElectraDecodersUtils.h
--- a/Engine/Plugins/Media/ElectraCodecs/Source/ElectraDecoders/Public/ElectraDecodersUtils.h
+++ b/Engine/Plugins/Media/ElectraCodecs/Source/ElectraDecoders/Public/ElectraDecodersUtils.h
@@ -79,6 +79,10 @@ namespace ElectraDecodersUtil
 		int32 NumberOfChannels = 0;
 	};
 
+	static constexpr uint32 Make4CC(const uint8 A, const uint8 B, const uint8 C, const uint8 D)
+	{
+		return (static_cast<uint32>(A) << 24) | (static_cast<uint32>(B) << 16) | (static_cast<uint32>(C) << 8) | static_cast<uint32>(D);
+	}
 
 	bool ELECTRADECODERS_API ParseMimeTypeWithCodec(FMimeTypeVideoCodecInfo& OutInfo, const FString& InMimeType);
 	bool ELECTRADECODERS_API ParseMimeTypeWithCodec(FMimeTypeAudioCodecInfo& OutInfo, const FString& InMimeType);
diff --git a/Engine/Plugins/Media/ElectraCodecs/Source/ElectraDecoders/Public/Utils/MPEG/ElectraUtilsMPEGVideo.h b/Engine/Plugins/Media/ElectraCodecs/Source/ElectraDecoders/Public/Utils/MPEG/ElectraUtilsMPEGVideo.h
--- a/Engine/Plugins/Media/ElectraCodecs/Source/ElectraDecoders/Public/Utils/MPEG/ElectraUtilsMPEGVideo.h
+++ b/Engine/Plugins/Media/ElectraCodecs/Source/ElectraDecoders/Public/Utils/MPEG/ElectraUtilsMPEGVideo.h
@@ -459,6 +459,7 @@ namespace ElectraDecodersUtil
 				PT_pic_timing = 1,
 				PT_user_data_registered_itu_t_t35 = 4,
 				// PT_tone_mapping_info = 23,							// not expected to be used (application custom)
+				PT_time_code = 136,
 				PT_mastering_display_colour_volume = 137,
 				// PT_chroma_resampling_filter_hint = 140,				// not expected to be used
 				// PT_knee_function_info = 141,							// not expected to be used (application custom)
@@ -489,6 +490,7 @@ namespace ElectraDecodersUtil
 			uint32 min_display_mastering_luminance = 0;
 		};
 		bool ELECTRADECODERS_API ParseFromSEIMessage(FSEImastering_display_colour_volume& OutMDCV, const FSEIMessage& InMessage);
+		bool ELECTRADECODERS_API ParseFromMDCVBox(FSEImastering_display_colour_volume& OutMDCV, const TConstArrayView<uint8>& InMDCVBox);
 
 		struct FSEIcontent_light_level_info
 		{
@@ -496,6 +498,8 @@ namespace ElectraDecodersUtil
 			uint16 max_pic_average_light_level = 0;
 		};
 		bool ELECTRADECODERS_API ParseFromSEIMessage(FSEIcontent_light_level_info& OutCLLI, const FSEIMessage& InMessage);
+		bool ELECTRADECODERS_API ParseFromCOLLBox(FSEIcontent_light_level_info& OutCLLI, const TConstArrayView<uint8>& InCOLLBox);
+		bool ELECTRADECODERS_API ParseFromCLLIBox(FSEIcontent_light_level_info& OutCLLI, const TConstArrayView<uint8>& InCLLIBox);
 
 		struct FSEIalternative_transfer_characteristics
 		{
@@ -534,6 +538,7 @@ namespace ElectraDecodersUtil
 			uint8 video_full_range_flag = 0;
 			uint8 video_format = 5;
 		};
+		bool ELECTRADECODERS_API ParseFromCOLRBox(FCommonColorimetry& OutColorimetry, const TConstArrayView<uint8>& InCOLRBox);
 
 		struct FCommonPictureTiming
 		{
diff --git a/Engine/Plugins/Media/ElectraPlayer/Source/ElectraPlayerRuntime/Private/Runtime/Decoder/VideoDecoderHelpers.cpp b/Engine/Plugins/Media/ElectraPlayer/Source/ElectraPlayerRuntime/Private/Runtime/Decoder/VideoDecoderHelpers.cpp
--- a/Engine/Plugins/Media/ElectraPlayer/Source/ElectraPlayerRuntime/Private/Runtime/Decoder/VideoDecoderHelpers.cpp
+++ b/Engine/Plugins/Media/ElectraPlayer/Source/ElectraPlayerRuntime/Private/Runtime/Decoder/VideoDecoderHelpers.cpp
@@ -207,6 +207,50 @@ void FHDRHelper::Update(int32 BitDepth, const FColorimetryHelper& InColorimetry,
 	bIsFirst = false;
 }
 
+void FHDRHelper::UpdateWith(const ElectraDecodersUtil::MPEG::FSEImastering_display_colour_volume& InSEI)
+{
+	if (!CurrentHDRInfo.IsValid())
+	{
+		CurrentHDRInfo = MakeShareable(new FVideoDecoderHDRInformation);
+	}
+	const float ST2086_Chroma_Scale = 50000.0f;
+	const float ST2086_Luma_Scale = 10000.0f;
+	FVideoDecoderHDRMetadata_mastering_display_colour_volume mdcv;
+	// The order in the HEVC SEI messages is G,B,R
+	mdcv.display_primaries_x[0] = InSEI.display_primaries_x[2] / ST2086_Chroma_Scale;
+	mdcv.display_primaries_y[0] = InSEI.display_primaries_y[2] / ST2086_Chroma_Scale;
+	mdcv.display_primaries_x[1] = InSEI.display_primaries_x[0] / ST2086_Chroma_Scale;
+	mdcv.display_primaries_y[1] = InSEI.display_primaries_y[0] / ST2086_Chroma_Scale;
+	mdcv.display_primaries_x[2] = InSEI.display_primaries_x[1] / ST2086_Chroma_Scale;
+	mdcv.display_primaries_y[2] = InSEI.display_primaries_y[1] / ST2086_Chroma_Scale;
+	mdcv.white_point_x = InSEI.white_point_x / ST2086_Chroma_Scale;
+	mdcv.white_point_y = InSEI.white_point_y / ST2086_Chroma_Scale;
+	mdcv.max_display_mastering_luminance = InSEI.max_display_mastering_luminance / ST2086_Luma_Scale;
+	mdcv.min_display_mastering_luminance = InSEI.min_display_mastering_luminance / ST2086_Luma_Scale;
+	CurrentHDRInfo->SetMasteringDisplayColourVolume(mdcv);
+}
+
+void FHDRHelper::UpdateWith(const ElectraDecodersUtil::MPEG::FSEIcontent_light_level_info& InSEI)
+{
+	if (!CurrentHDRInfo.IsValid())
+	{
+		CurrentHDRInfo = MakeShareable(new FVideoDecoderHDRInformation);
+	}
+	FVideoDecoderHDRMetadata_content_light_level_info ll;
+	ll.max_content_light_level = InSEI.max_content_light_level;
+	ll.max_pic_average_light_level = InSEI.max_pic_average_light_level;
+	CurrentHDRInfo->SetContentLightLevelInfo(ll);
+}
+
+void FHDRHelper::UpdateWith(const ElectraDecodersUtil::MPEG::FSEIalternative_transfer_characteristics& InSEI)
+{
+	if (!CurrentHDRInfo.IsValid())
+	{
+		CurrentHDRInfo = MakeShareable(new FVideoDecoderHDRInformation);
+	}
+	CurrentAlternativeTransferCharacteristics = InSEI.preferred_transfer_characteristics;
+}
+
 
 void FHDRHelper::UpdateFromMPEGBoxes(int32 BitDepth, const FColorimetryHelper& InColorimetry, const TArray<uint8>& InMDCVBox, const TArray<uint8>& InCLLIBox)
 {
diff --git a/Engine/Plugins/Media/ElectraPlayer/Source/ElectraPlayerRuntime/Public/Decoder/VideoDecoderHelpers.h b/Engine/Plugins/Media/ElectraPlayer/Source/ElectraPlayerRuntime/Public/Decoder/VideoDecoderHelpers.h
--- a/Engine/Plugins/Media/ElectraPlayer/Source/ElectraPlayerRuntime/Public/Decoder/VideoDecoderHelpers.h
+++ b/Engine/Plugins/Media/ElectraPlayer/Source/ElectraPlayerRuntime/Public/Decoder/VideoDecoderHelpers.h
@@ -68,6 +68,10 @@ public:
 	ELECTRAPLAYERRUNTIME_API void Update(int32 BitDepth, const FColorimetryHelper& InColorimetry, const TArray<ElectraDecodersUtil::MPEG::FSEIMessage>& InGlobalPrefixSEIs, const TArray<ElectraDecodersUtil::MPEG::FSEIMessage>& InLocalPrefixSEIs, bool bIsNewCLVS);
 	ELECTRAPLAYERRUNTIME_API void UpdateFromMPEGBoxes(int32 BitDepth, const FColorimetryHelper& InColorimetry, const TArray<uint8>& InMDCVBox, const TArray<uint8>& InCLLIBox);
 	ELECTRAPLAYERRUNTIME_API void Update(int32 BitDepth, const FColorimetryHelper& InColorimetry, const TOptional<FVideoDecoderHDRMetadata_mastering_display_colour_volume>& InMDCV, const TOptional<FVideoDecoderHDRMetadata_content_light_level_info>& InCLLI);
+	ELECTRAPLAYERRUNTIME_API void UpdateWith(const ElectraDecodersUtil::MPEG::FSEImastering_display_colour_volume& InSEI);
+	ELECTRAPLAYERRUNTIME_API void UpdateWith(const ElectraDecodersUtil::MPEG::FSEIcontent_light_level_info& InSEI);
+	ELECTRAPLAYERRUNTIME_API void UpdateWith(const ElectraDecodersUtil::MPEG::FSEIalternative_transfer_characteristics& InSEI);
+	ELECTRAPLAYERRUNTIME_API void SetHDRType(int32 BitDepth, const FColorimetryHelper& InColorimetry);
 	ELECTRAPLAYERRUNTIME_API void UpdateParamDict(FParamDict& InOutDictionary);
 private:
 	class FVideoDecoderHDRInformation : public IVideoDecoderHDRInformation
@@ -91,7 +95,6 @@ private:
 		TOptional<FVideoDecoderHDRMetadata_mastering_display_colour_volume> MasteringDisplayColourVolume;
 		TOptional<FVideoDecoderHDRMetadata_content_light_level_info> ContentLightLevelInfo;
 	};
-	void SetHDRType(int32 BitDepth, const FColorimetryHelper& InColorimetry);
 	TSharedPtr<FVideoDecoderHDRInformation, ESPMode::ThreadSafe> CurrentHDRInfo;
 	TOptional<ElectraDecodersUtil::MPEG::FSEIMessage> ActiveMasteringDisplayColourVolume;
 	TOptional<ElectraDecodersUtil::MPEG::FSEIMessage> ActiveContentLightLevelInfo;
diff --git a/Engine/Plugins/Media/ElectraUtil/Source/ElectraBase/Private/ParameterDictionary.cpp b/Engine/Plugins/Media/ElectraUtil/Source/ElectraBase/Private/ParameterDictionary.cpp
--- a/Engine/Plugins/Media/ElectraUtil/Source/ElectraBase/Private/ParameterDictionary.cpp
+++ b/Engine/Plugins/Media/ElectraUtil/Source/ElectraBase/Private/ParameterDictionary.cpp
@@ -61,6 +61,18 @@ FVariantValue::FVariantValue(const FTimeValue& TimeValue)
 	Set(TimeValue);
 }
 
+FVariantValue::FVariantValue(const FTimespan& TimespanValue)
+	: DataType(EDataType::TypeUninitialized)
+{
+	Set(TimespanValue);
+}
+
+FVariantValue::FVariantValue(const FTimecode& TimecodeValue)
+	: DataType(EDataType::TypeUninitialized)
+{
+	Set(TimecodeValue);
+}
+
 FVariantValue::FVariantValue(void* PointerValue)
 	: DataType(EDataType::TypeUninitialized)
 {
@@ -79,26 +91,50 @@ void FVariantValue::CopyInternal(const FVariantValue& FromOther)
 	switch(FromOther.DataType)
 	{
 		case EDataType::TypeUninitialized:
+		{
 			Clear();
 			break;
+		}
 		case EDataType::TypeFString:
+		{
 			Set(FromOther.GetFString());
 			break;
+		}
 		case EDataType::TypeDouble:
+		{
 			Set(FromOther.GetDouble());
 			break;
+		}
 		case EDataType::TypeInt64:
+		{
 			Set(FromOther.GetInt64());
 			break;
+		}
 		case EDataType::TypeBoolean:
+		{
 			Set(FromOther.GetBool());
 			break;
+		}
 		case EDataType::TypeTimeValue:
+		{
 			Set(FromOther.GetTimeValue());
 			break;
+		}
+		case EDataType::TypeTimespanValue:
+		{
+			Set(FromOther.GetTimespan());
+			break;
+		}
+		case EDataType::TypeTimecodeValue:
+		{
+			Set(FromOther.GetTimecode());
+			break;
+		}
 		case EDataType::TypeVoidPointer:
+		{
 			Set(FromOther.GetPointer());
 			break;
+		}
 		case EDataType::TypeSharedPointer:
 		{
 			Clear();
@@ -109,12 +145,16 @@ void FVariantValue::CopyInternal(const FVariantValue& FromOther)
 			break;
 		}
 		case EDataType::TypeU8Array:
+		{
 			Set(FromOther.GetArray());
 			break;
+		}
 		default:
+		{
 			Clear();
 			check(!"Whoops");
 			break;
+		}
 	}
 }
 
@@ -122,23 +162,45 @@ FVariant FVariantValue::ToFVariant() const
 {
 	switch(GetDataType())
 	{
-	case EDataType::TypeFString:
-		return FVariant(GetFString());
-	case EDataType::TypeDouble:
-		return FVariant(GetDouble());
-	case EDataType::TypeInt64:
-		return FVariant(GetInt64());
-	case EDataType::TypeBoolean:
-		return FVariant(GetBool());
-	case EDataType::TypeTimeValue:
-		return FVariant(GetTimeValue().GetAsTimespan());
-	case EDataType::TypeVoidPointer:
-		return FVariant(reinterpret_cast<uint64>(GetPointer()));
-	case EDataType::TypeU8Array:
-		return FVariant(GetArray());
-	case EDataType::TypeSharedPointer: // Can't be converted.
+		case EDataType::TypeFString:
+		{
+			return FVariant(GetFString());
+		}
+		case EDataType::TypeDouble:
+		{
+			return FVariant(GetDouble());
+		}
+		case EDataType::TypeInt64:
+		{
+			return FVariant(GetInt64());
+		}
+		case EDataType::TypeBoolean:
+		{
+			return FVariant(GetBool());
+		}
+		case EDataType::TypeTimeValue:
+		{
+			return FVariant(GetTimeValue().GetAsTimespan());
+		}
+		case EDataType::TypeTimespanValue:
+		{
+			return FVariant(GetTimespan());
+		}
+		case EDataType::TypeVoidPointer:
+		{
+			return FVariant(reinterpret_cast<uint64>(GetPointer()));
+		}
+		case EDataType::TypeU8Array:
+		{
+			return FVariant(GetArray());
+		}
+		// Types that can't be converted.
+		case EDataType::TypeTimecodeValue:
+		case EDataType::TypeSharedPointer:
 		default:
+		{
 			return FVariant();
+		}
 	}
 }
 
@@ -158,12 +220,26 @@ void FVariantValue::Clear()
 			Time->~FTimeValue();
 			break;
 		}
+		case EDataType::TypeTimespanValue:
+		{
+			FTimespan* Timespan = reinterpret_cast<FTimespan*>(&DataBuffer);
+			Timespan->~FTimespan();
+			break;
+		}
+		case EDataType::TypeTimecodeValue:
+		{
+			FTimecode* Timecode = reinterpret_cast<FTimecode*>(&DataBuffer);
+			Timecode->~FTimecode();
+			break;
+		}
 		case EDataType::TypeUninitialized:
 		case EDataType::TypeDouble:
 		case EDataType::TypeInt64:
 		case EDataType::TypeBoolean:
 		case EDataType::TypeVoidPointer:
+		{
 			break;
+		}
 		case EDataType::TypeSharedPointer:
 		{
 			FSharedPtrHolderBase* Pointer = reinterpret_cast<FSharedPtrHolderBase*>(&DataBuffer);
@@ -231,6 +307,24 @@ FVariantValue& FVariantValue::Set(const FTimeValue& TimeValue)
 	return *this;
 }
 
+FVariantValue& FVariantValue::Set(const FTimespan& TimespanValue)
+{
+	Clear();
+	FTimespan* ValuePtr = reinterpret_cast<FTimespan*>(&DataBuffer);
+	*ValuePtr = TimespanValue;
+	DataType = EDataType::TypeTimespanValue;
+	return *this;
+}
+
+FVariantValue& FVariantValue::Set(const FTimecode& TimecodeValue)
+{
+	Clear();
+	FTimecode* ValuePtr = reinterpret_cast<FTimecode*>(&DataBuffer);
+	*ValuePtr = TimecodeValue;
+	DataType = EDataType::TypeTimecodeValue;
+	return *this;
+}
+
 FVariantValue& FVariantValue::Set(void* PointerValue)
 {
 	Clear();
@@ -319,6 +413,34 @@ const FTimeValue& FVariantValue::GetTimeValue() const
 	}
 }
 
+const FTimespan& FVariantValue::GetTimespan() const
+{
+	if (ensure(DataType == EDataType::TypeTimespanValue))
+	{
+		const FTimespan* Timespan = reinterpret_cast<const FTimespan*>(&DataBuffer);
+		return *Timespan;
+	}
+	else
+	{
+		static FTimespan Empty;
+		return Empty;
+	}
+}
+
+const FTimecode& FVariantValue::GetTimecode() const
+{
+	if (ensure(DataType == EDataType::TypeTimecodeValue))
+	{
+		const FTimecode* Timecode = reinterpret_cast<const FTimecode*>(&DataBuffer);
+		return *Timecode;
+	}
+	else
+	{
+		static FTimecode Empty;
+		return Empty;
+	}
+}
+
 void* const & FVariantValue::GetPointer() const
 {
 	if (ensure(DataType == EDataType::TypeVoidPointer))
@@ -398,6 +520,27 @@ FTimeValue FVariantValue::SafeGetTimeValue(const FTimeValue& Default) const
 	return Default;
 }
 
+FTimespan FVariantValue::SafeGetTimespan(const FTimespan& Default) const
+{
+	if (DataType == EDataType::TypeTimespanValue)
+	{
+		const FTimespan* Timespan = reinterpret_cast<const FTimespan*>(&DataBuffer);
+		return *Timespan;
+	}
+	return Default;
+}
+
+FTimecode FVariantValue::SafeGetTimecode(const FTimecode& Default) const
+{
+	if (DataType == EDataType::TypeTimecodeValue)
+	{
+		const FTimecode* Timecode = reinterpret_cast<const FTimecode*>(&DataBuffer);
+		return *Timecode;
+	}
+	return Default;
+}
+
+
 void* FVariantValue::SafeGetPointer(void* Default) const
 {
 	if (DataType == EDataType::TypeVoidPointer)
@@ -465,12 +608,12 @@ FVariantValue FParamDict::GetValue(const FName& Key) const
 
 void FParamDict::Remove(const FName& Key)
 {
-	Dictionary.Remove(Key); 
+	Dictionary.Remove(Key);
 }
 
 void FParamDict::Set(const FName& Key, const FVariantValue& Value)
-{ 
-	Dictionary.Emplace(Key, Value); 
+{
+	Dictionary.Emplace(Key, Value);
 }
 
 void FParamDict::Set(const FName& Key, FVariantValue&& Value)
diff --git a/Engine/Plugins/Media/ElectraUtil/Source/ElectraBase/Private/PlayerTime.cpp b/Engine/Plugins/Media/ElectraUtil/Source/ElectraBase/Private/PlayerTime.cpp
--- a/Engine/Plugins/Media/ElectraUtil/Source/ElectraBase/Private/PlayerTime.cpp
+++ b/Engine/Plugins/Media/ElectraUtil/Source/ElectraBase/Private/PlayerTime.cpp
@@ -253,6 +253,48 @@ namespace Electra
 		return bIsNeg ? -r : r;
 	}
 
+	FTimespan FTimeFraction::GetAsTimespan() const
+	{
+		int64 Ticks = 0;
+		if (bIsValid)
+		{
+			if (Denominator)
+			{
+				if (Numerator != 0)
+				{
+					if (Denominator == ETimespan::TicksPerSecond)
+					{
+						Ticks = Numerator;
+					}
+					else if (Numerator >= -922337203685LL && Numerator <= 922337203685LL)
+					{
+						static_assert(ETimespan::TicksPerSecond == 10000000); // otherwise the constants above are wrong
+						Ticks = Numerator * ETimespan::TicksPerSecond / Denominator;
+					}
+					else
+					{
+						bool bIsNeg = Numerator < 0;
+						TBigInt<128> n(bIsNeg ? -Numerator : Numerator);
+						TBigInt<128> d(Denominator);
+						TBigInt<128> s(ETimespan::TicksPerSecond);
+						n *= s;
+						n /= d;
+						Ticks = n.ToInt();
+						if (bIsNeg)
+						{
+							Ticks = -Ticks;
+						}
+					}
+				}
+			}
+			else
+			{
+				Ticks = Numerator < 0 ? ETimespan::MinTicks : ETimespan::MaxTicks;
+			}
+		}
+		return FTimespan(Ticks);
+	}
+
 
 }
 
diff --git a/Engine/Plugins/Media/ElectraUtil/Source/ElectraBase/Public/ParameterDictionary.h b/Engine/Plugins/Media/ElectraUtil/Source/ElectraBase/Public/ParameterDictionary.h
--- a/Engine/Plugins/Media/ElectraUtil/Source/ElectraBase/Public/ParameterDictionary.h
+++ b/Engine/Plugins/Media/ElectraUtil/Source/ElectraBase/Public/ParameterDictionary.h
@@ -8,6 +8,8 @@
 #include "Templates/TypeCompatibleBytes.h"
 #include "Containers/Array.h"
 #include "Misc/Variant.h"
+#include "Misc/Timespan.h"
+#include "Misc/Timecode.h"
 
 
 namespace Electra
@@ -50,6 +52,8 @@ public:
 	explicit FVariantValue(int64 Int64Value);
 	explicit FVariantValue(bool BoolValue);
 	explicit FVariantValue(const FTimeValue& TimeValue);
+	explicit FVariantValue(const FTimespan& TimespanValue);
+	explicit FVariantValue(const FTimecode& TimecodeValue);
 	explicit FVariantValue(void* PointerValue);
 	template<typename T> explicit FVariantValue(const TSharedPtr<T, ESPMode::ThreadSafe>& PointerValue)
 	: DataType(EDataType::TypeUninitialized)
@@ -63,6 +67,8 @@ public:
 	FVariantValue& Set(int64 Int64Value);
 	FVariantValue& Set(bool BoolValue);
 	FVariantValue& Set(const FTimeValue& TimeValue);
+	FVariantValue& Set(const FTimespan& TimespanValue);
+	FVariantValue& Set(const FTimecode& TimecodeValue);
 	FVariantValue& Set(void* PointerValue);
 	template <typename T> FVariantValue& Set(const TSharedPtr<T, ESPMode::ThreadSafe>& PointerValue)
 	{
@@ -79,6 +85,8 @@ public:
 	const int64& GetInt64() const;
 	const bool& GetBool() const;
 	const FTimeValue& GetTimeValue() const;
+	const FTimespan& GetTimespan() const;
+	const FTimecode& GetTimecode() const;
 	void* const & GetPointer() const;
 	template<typename T> TSharedPtr<T, ESPMode::ThreadSafe> GetSharedPointer() const
 	{
@@ -100,6 +108,8 @@ public:
 	int64 SafeGetInt64(int64 Default=0) const;
 	bool SafeGetBool(bool Default=false) const;
 	FTimeValue SafeGetTimeValue(const FTimeValue& Default=FTimeValue()) const;
+	FTimespan SafeGetTimespan(const FTimespan& Default=FTimespan()) const;
+	FTimecode SafeGetTimecode(const FTimecode& Default=FTimecode()) const;
 	void* SafeGetPointer(void* Default=nullptr) const;
 	const TArray<uint8>& SafeGetArray() const;
 
@@ -111,6 +121,8 @@ public:
 		TypeInt64,
 		TypeBoolean,
 		TypeTimeValue,
+		TypeTimespanValue,
+		TypeTimecodeValue,
 		TypeVoidPointer,
 		TypeSharedPointer,
 		TypeU8Array,
@@ -142,6 +154,8 @@ private:
 		uint8 MemSizeVoidPtr[sizeof(void*)];
 		uint8 MemSizeBoolean[sizeof(bool)];
 		uint8 MemSizeTimeValue[sizeof(FTimeValue)];
+		uint8 MemSizeTimespan[sizeof(FTimespan)];
+		uint8 MemSizeTimecode[sizeof(FTimecode)];
 		uint8 MemSizeSharedPtrValue[sizeof(TSharedPtrHolder<uint8>)];
 		uint8 MemSizeTArray[sizeof(TArray<uint8>)];
 	};
@@ -150,7 +164,7 @@ private:
 	void CopyInternal(const FVariantValue& FromOther);
 
 	TAlignedBytes<sizeof(FUnionLayout), 16> DataBuffer;
-	EDataType	DataType;
+	EDataType DataType;
 };
 
 
diff --git a/Engine/Plugins/Media/ElectraUtil/Source/ElectraBase/Public/PlayerTime.h b/Engine/Plugins/Media/ElectraUtil/Source/ElectraBase/Public/PlayerTime.h
--- a/Engine/Plugins/Media/ElectraUtil/Source/ElectraBase/Public/PlayerTime.h
+++ b/Engine/Plugins/Media/ElectraUtil/Source/ElectraBase/Public/PlayerTime.h
@@ -601,6 +601,12 @@ public:
 		return Zero;
 	}
 
+	static const FTimeFraction& GetPositiveInfinity()
+	{
+		static FTimeFraction Inf(0, 0);
+		return Inf;
+	}
+
 	FTimeFraction() : Numerator(0), Denominator(0), bIsValid(false)
 	{
 	}
@@ -618,12 +624,17 @@ public:
 		SetFromTimeValue(tv);
 	}
 
+	FTimeFraction(const FTimespan& tv)
+	{
+		SetFromTimespan(tv);
+	}
+
 
 	FTimeFraction& operator=(const FTimeFraction& rhs)
 	{
-		Numerator   = rhs.Numerator;
+		Numerator = rhs.Numerator;
 		Denominator = rhs.Denominator;
-		bIsValid	= rhs.bIsValid;
+		bIsValid = rhs.bIsValid;
 		return *this;
 	}
 
@@ -655,11 +666,21 @@ public:
 	//! Returns this time value in a custom timebase. Requires internal bigint conversion and is therefor SLOW!
 	ELECTRABASE_API int64 GetAsTimebase(uint32 CustomTimebase) const;
 
+	ELECTRABASE_API FTimespan GetAsTimespan() const;
+
 	FTimeFraction& SetFromND(int64 InNumerator, uint32 InDenominator)
 	{
-		Numerator   = InNumerator;
+		Numerator = InNumerator;
 		Denominator = InDenominator;
-		bIsValid	= true;
+		bIsValid = true;
+		return *this;
+	}
+
+	FTimeFraction& SetToPositiveInfinity()
+	{
+		Numerator = 0;
+		Denominator = 0;
+		bIsValid = true;
 		return *this;
 	}
 
@@ -680,6 +701,14 @@ public:
 		return *this;
 	}
 
+	FTimeFraction& SetFromTimespan(const FTimespan& tv)
+	{
+		Numerator = tv.GetTicks();
+		Denominator = (Numerator == ETimespan::MinTicks || Numerator == ETimespan::MaxTicks) ? 0U : (uint32)ETimespan::TicksPerSecond;
+		bIsValid = true;
+		return *this;
+	}
+
 	ELECTRABASE_API FTimeFraction& SetFromFloatString(const FString& In);
 
 
