ElectraBase: New optimized mp4 box and track handler

#rnx

diff --git a/Engine/Plugins/Media/ElectraUtil/Source/ElectraBase/Private/Utilities/MP4Boxes/MP4Boxes.cpp b/Engine/Plugins/Media/ElectraUtil/Source/ElectraBase/Private/Utilities/MP4Boxes/MP4Boxes.cpp
new file mode 100644
--- /dev/null
+++ b/Engine/Plugins/Media/ElectraUtil/Source/ElectraBase/Private/Utilities/MP4Boxes/MP4Boxes.cpp
@@ -0,0 +1,2606 @@
+// Copyright Epic Games, Inc. All Rights Reserved.
+
+#include "Utilities/MP4Boxes/MP4Boxes.h"
+#include "Utilities/UtilitiesMP4.h"
+#include "Utilities/ISO639-Map.h"
+#include "ElectraBaseModule.h"	// for log category `LogElectraBase`
+
+namespace Electra
+{
+	namespace UtilitiesMP4
+	{
+
+
+		void FMP4BoxBase::ProcessBoxChildrenRecursively(FMP4AtomReaderHelper& InOutReader, const FMP4BoxInfo& InCurrentBoxInfo)
+		{
+			FMP4BoxInfo bi;
+			while(InOutReader.ParseIntoBoxInfo(bi, InCurrentBoxInfo.Offset + InCurrentBoxInfo.DataOffset + InOutReader.GetCurrentOffset()))
+			{
+				TSharedPtr<FMP4BoxBase> Child = FMP4BoxFactory::Get().Create(AsWeak(), bi);
+				if (AddChildBox(Child) && !Child->IsLeafBox() && !Child->IsListOfEntries())
+				{
+					FMP4AtomReaderHelper ar(Child->BoxInfo.Data);
+					Child->ProcessBoxChildrenRecursively(ar, Child->BoxInfo);
+				}
+				InOutReader.SkipBytes(bi.Size - bi.DataOffset);
+			}
+		}
+
+/****************************************************************************************************************************************************/
+		struct FMP4BoxMVHD::FParsed
+		{
+			uint64 Duration = 0;
+			uint32 Flags = 0;
+			uint32 Timescale = 0;
+			uint8 Version = 0;
+		};
+		void FMP4BoxMVHD::ParseIfRequired()
+		{
+			if (Parsed.IsValid())
+			{
+				return;
+			}
+			Parsed = MakePimpl<FParsed>();
+			FMP4AtomReaderHelper ar(BoxInfo.Data);
+			ar.ReadVersionAndFlags(Parsed->Version, Parsed->Flags);
+			if (Parsed->Version == 1)
+			{
+				ar.SkipBytes(16);				// `creation_time` and `modification_time`
+				ar.Read(Parsed->Timescale);
+				ar.Read(Parsed->Duration);
+			}
+			else
+			{
+				ar.SkipBytes(8);				// `creation_time` and `modification_time`
+				ar.Read(Parsed->Timescale);
+				uint32 Dur;
+				ar.Read(Dur);
+				Parsed->Duration = Dur == TNumericLimits<uint32>::Max() ? TNumericLimits<uint64>::Max() : Dur;
+			}
+		}
+		FTimeFraction FMP4BoxMVHD::GetDuration()
+		{
+			ParseIfRequired();
+			// If not specified (all 1's) or the numerator won't fit into an int64 of the FTimeFraction
+			// we return an invalid time.
+			if (Parsed->Duration == TNumericLimits<uint64>::Max() || Parsed->Duration > 0x7fffffffffffffffULL)
+			{
+				return FTimeFraction::GetInvalid();
+			}
+			return FTimeFraction((int64) Parsed->Duration, Parsed->Timescale);
+		}
+		uint32 FMP4BoxMVHD::GetTimescale()
+		{
+			ParseIfRequired();
+			return Parsed->Timescale;
+		}
+
+/****************************************************************************************************************************************************/
+		struct FMP4BoxTKHD::FParsed
+		{
+			int64 Duration = 0;
+			uint32 Flags = 0;
+			uint32 TrackID = 0;
+			uint32 Width = 0;
+			uint32 Height = 0;
+			int16 Layer = 0;
+			int16 AlternateGroup = 0;
+			uint8 Version = 0;
+		};
+		void FMP4BoxTKHD::ParseIfRequired()
+		{
+			if (Parsed.IsValid())
+			{
+				return;
+			}
+			Parsed = MakePimpl<FParsed>();
+			FMP4AtomReaderHelper ar(BoxInfo.Data);
+			ar.ReadVersionAndFlags(Parsed->Version, Parsed->Flags);
+			if (Parsed->Version == 1)
+			{
+				ar.SkipBytes(16);				// `creation_time` and `modification_time`
+				ar.Read(Parsed->TrackID);
+				ar.SkipBytes(4);				// `reserved`
+				uint64 Dur;
+				ar.Read(Dur);
+				Parsed->Duration = Dur > (uint64)TNumericLimits<int64>::Max() ? TNumericLimits<int64>::Max() : (int64)Dur;
+			}
+			else
+			{
+				ar.SkipBytes(8);				// `creation_time` and `modification_time`
+				ar.Read(Parsed->TrackID);
+				ar.SkipBytes(4);				// `reserved`
+				uint32 Dur;
+				ar.Read(Dur);
+				Parsed->Duration = Dur == TNumericLimits<uint32>::Max() ? TNumericLimits<int64>::Max() : (int64)Dur;
+			}
+			ar.SkipBytes(8);					// `reserved` (uint32 * 2)
+			uint16 Value16;
+			ar.Read(Value16);
+			Parsed->Layer = (int16) Value16;
+			ar.Read(Value16);
+			Parsed->AlternateGroup = (int16) Value16;
+			ar.SkipBytes(4 + 9*4);				// `volume`, `reserved` and `matrix`
+			ar.Read(Parsed->Width);
+			ar.Read(Parsed->Height);
+		}
+
+		uint32 FMP4BoxTKHD::GetFlags()
+		{
+			ParseIfRequired();
+			return Parsed->Flags;
+		}
+
+		int64 FMP4BoxTKHD::GetDuration()
+		{
+			ParseIfRequired();
+			return Parsed->Duration;
+		}
+
+		uint32 FMP4BoxTKHD::GetTrackID()
+		{
+			ParseIfRequired();
+			return Parsed->TrackID;
+		}
+
+		uint16 FMP4BoxTKHD::GetWidth()
+		{
+			ParseIfRequired();
+			return Parsed->Width >> 16;
+		}
+
+		uint16 FMP4BoxTKHD::GetHeight()
+		{
+			ParseIfRequired();
+			return Parsed->Height >> 16;
+		}
+
+/****************************************************************************************************************************************************/
+
+		struct FMP4BoxTREF::FParsed
+		{
+			TArray<FEntry> Entries;
+			uint32 Flags = 0;
+			uint8 Version = 0;
+		};
+		void FMP4BoxTREF::ParseIfRequired()
+		{
+			if (Parsed.IsValid())
+			{
+				return;
+			}
+			Parsed = MakePimpl<FParsed>();
+			FMP4AtomReaderHelper ar(BoxInfo.Data);
+			FMP4BoxInfo bi;
+			while(ar.ParseIntoBoxInfo(bi, BoxInfo.Offset + BoxInfo.DataOffset + ar.GetCurrentOffset()))
+			{
+				FEntry& e = Parsed->Entries.Emplace_GetRef();
+				e.Type = bi.Type;
+				int32 NumReferences = (bi.Size - bi.DataOffset) / 4;
+				e.TrackIDs.SetNumUninitialized(NumReferences);
+				for(int32 i=0; i<NumReferences; ++i)
+				{
+					ar.Read(e.TrackIDs[i]);
+				}
+			}
+		}
+
+		const TArray<FMP4BoxTREF::FEntry>& FMP4BoxTREF::GetEntries()
+		{
+			ParseIfRequired();
+			return Parsed->Entries;
+		}
+		TArray<FMP4BoxTREF::FEntry> FMP4BoxTREF::GetEntriesOfType(uint32 InType)
+		{
+			ParseIfRequired();
+			TArray<FMP4BoxTREF::FEntry> Entries;
+			for(auto& e : Parsed->Entries)
+			{
+				if (e.Type == InType)
+				{
+					Entries.Emplace(e);
+				}
+			}
+			return Entries;
+		}
+
+/****************************************************************************************************************************************************/
+
+		struct FMP4BoxELST::FParsed
+		{
+			uint32 Flags = 0;
+			uint8 Version = 0;
+			TArray<FEntry> Entries;
+		};
+		void FMP4BoxELST::ParseIfRequired()
+		{
+			if (Parsed.IsValid())
+			{
+				return;
+			}
+			Parsed = MakePimpl<FParsed>();
+			FMP4AtomReaderHelper ar(BoxInfo.Data);
+			ar.ReadVersionAndFlags(Parsed->Version, Parsed->Flags);
+			uint32 NumEntries;
+			ar.Read(NumEntries);
+			Parsed->Entries.AddDefaulted((int32) NumEntries);
+			if (Parsed->Version == 1)
+			{
+				for(uint32 i=0; i<NumEntries; ++i)
+				{
+					FEntry& e(Parsed->Entries[i]);
+					ar.Read(e.edit_duration);
+					ar.Read(e.media_time);
+					ar.Read(e.media_rate_integer);
+					ar.Read(e.media_rate_fraction);
+				}
+			}
+			else
+			{
+				uint32 edit_duration;
+				int32 media_time;
+				for(uint32 i=0; i<NumEntries; ++i)
+				{
+					FEntry& e(Parsed->Entries[i]);
+					ar.Read(edit_duration);
+					e.edit_duration = edit_duration;
+					ar.Read(media_time);
+					e.media_time = media_time;
+					ar.Read(e.media_rate_integer);
+					ar.Read(e.media_rate_fraction);
+				}
+			}
+		}
+
+		const TArray<FMP4BoxELST::FEntry>& FMP4BoxELST::GetEntries()
+		{
+			ParseIfRequired();
+			return Parsed->Entries;
+		}
+		bool FMP4BoxELST::RepeatEdits()
+		{
+			ParseIfRequired();
+			return !!(Parsed->Flags & 1);
+		}
+
+/****************************************************************************************************************************************************/
+
+		struct FMP4BoxMDHD::FParsed
+		{
+			FString Lang639_1;
+			int64 Duration = 0;
+			uint32 Flags = 0;
+			uint32 Timescale = 0;
+			uint8 Version = 0;
+		};
+		void FMP4BoxMDHD::ParseIfRequired()
+		{
+			if (Parsed.IsValid())
+			{
+				return;
+			}
+			Parsed = MakePimpl<FParsed>();
+			FMP4AtomReaderHelper ar(BoxInfo.Data);
+			ar.ReadVersionAndFlags(Parsed->Version, Parsed->Flags);
+			if (Parsed->Version == 1)
+			{
+				ar.SkipBytes(16);				// `creation_time` and `modification_time`
+				ar.Read(Parsed->Timescale);
+				ar.Read(Parsed->Duration);
+				Parsed->Duration = Parsed->Duration < 0 ? TNumericLimits<int64>::Max() : Parsed->Duration;
+			}
+			else
+			{
+				ar.SkipBytes(8);				// `creation_time` and `modification_time`
+				uint32 duration;
+				ar.Read(Parsed->Timescale);
+				ar.Read(duration);
+				Parsed->Duration = duration;
+			}
+
+			uint16 Value16;
+			ar.Read(Value16);
+			char Language[3] { (char)(0x60 + ((Value16 & 0x7c00) >> 10)), (char)(0x60 + ((Value16 & 0x03e0) >> 5)), (char)(0x60 + (Value16 & 0x001f)) };
+			// Try to map the ISO-639-2T language code to the shorter ISO-639-1 code if possible.
+			Parsed->Lang639_1 = ISO639::MapTo639_1(FString::ConstructFromPtrSize(Language, UE_ARRAY_COUNT(Language)));
+		}
+
+		FTimeFraction FMP4BoxMDHD::GetDuration()
+		{
+			ParseIfRequired();
+			return FTimeFraction((int64) Parsed->Duration, Parsed->Timescale);
+		}
+		uint32 FMP4BoxMDHD::GetTimescale()
+		{
+			ParseIfRequired();
+			return Parsed->Timescale;
+		}
+		const FString& FMP4BoxMDHD::GetLanguage()
+		{
+			ParseIfRequired();
+			return Parsed->Lang639_1;
+		}
+
+/****************************************************************************************************************************************************/
+
+		struct FMP4BoxHDLR::FParsed
+		{
+			FString HandlerName;
+			uint32 HandlerType = 0;
+			uint32 Flags = 0;
+			uint8 Version = 0;
+		};
+		void FMP4BoxHDLR::ParseIfRequired()
+		{
+			if (Parsed.IsValid())
+			{
+				return;
+			}
+			Parsed = MakePimpl<FParsed>();
+			FMP4AtomReaderHelper ar(BoxInfo.Data);
+			ar.ReadVersionAndFlags(Parsed->Version, Parsed->Flags);
+			ar.SkipBytes(4);				// `pre_defined`
+			ar.Read(Parsed->HandlerType);
+			ar.SkipBytes(12);				// `reserved`
+			ar.ReadStringUTF8(Parsed->HandlerName, ar.GetNumBytesRemaining());
+		}
+
+		uint32 FMP4BoxHDLR::GetHandlerType()
+		{
+			ParseIfRequired();
+			return Parsed->HandlerType;
+		}
+		FString FMP4BoxHDLR::GetHandlerName()
+		{
+			ParseIfRequired();
+			return Parsed->HandlerName;
+		}
+
+/****************************************************************************************************************************************************/
+
+		struct FMP4BoxELNG::FParsed
+		{
+			FString Language;
+			uint32 Flags = 0;
+			uint8 Version = 0;
+		};
+		void FMP4BoxELNG::ParseIfRequired()
+		{
+			if (Parsed.IsValid())
+			{
+				return;
+			}
+			Parsed = MakePimpl<FParsed>();
+			FMP4AtomReaderHelper ar(BoxInfo.Data);
+			ar.ReadVersionAndFlags(Parsed->Version, Parsed->Flags);
+			ar.ReadStringUTF8(Parsed->Language, ar.GetNumBytesRemaining());
+		}
+
+		FString FMP4BoxELNG::GetLanguage()
+		{
+			ParseIfRequired();
+			return Parsed->Language;
+		}
+
+/****************************************************************************************************************************************************/
+
+		struct FMP4BoxSTSD::FParsed
+		{
+			FMP4BoxSampleEntry::ESampleType SampleType = FMP4BoxSampleEntry::ESampleType::Unsupported;
+			uint32 Flags = 0;
+			uint8 Version = 0;
+		};
+		void FMP4BoxSTSD::ParseIfRequired()
+		{
+			if (Parsed.IsValid())
+			{
+				return;
+			}
+			Parsed = MakePimpl<FParsed>();
+			FMP4AtomReaderHelper ar(BoxInfo.Data);
+			ar.ReadVersionAndFlags(Parsed->Version, Parsed->Flags);
+
+			// In order to parse the sample entry we need to know the media handler type.
+			// Structurally this `stsd` box must be contained in a path `mdia`->`minf`->`stbl`->`stsd`, with
+			// the required `hdlr` box being given under `mdia`, so we need to get our enclosing `mdia` box first.
+			auto MDIABox = FindParentBox<FMP4BoxMDIA>(MakeBoxAtom('m','d','i','a'));
+			if (!MDIABox.IsValid())
+			{
+				UE_LOG(LogElectraBase, Error, TEXT("Could not find the parent `mdia` box to parse this `stsd` box!"));
+				return;
+			}
+			auto HDLRBox = MDIABox->FindBoxRecursive<FMP4BoxHDLR>(Electra::UtilitiesMP4::MakeBoxAtom('h','d','l','r'), 1);
+			if (!HDLRBox.IsValid())
+			{
+				// If a handler is missing for whatever reason we could look for a `vmhd`, `smhd`, etc. box under the `minf` box.
+				UE_LOG(LogElectraBase, Error, TEXT("Could not find the corresponding `hdlr` box required to parse this `stsd` box!"));
+				return;
+			}
+			switch(HDLRBox->GetHandlerType())
+			{
+				case MakeBoxAtom('v','i','d','e'):
+				{
+					Parsed->SampleType = FMP4BoxSampleEntry::ESampleType::Video;
+					break;
+				}
+				case MakeBoxAtom('s','o','u','n'):
+				{
+					Parsed->SampleType = FMP4BoxSampleEntry::ESampleType::Audio;
+					break;
+				}
+				case MakeBoxAtom('t','m','c','d'):
+				{
+					Parsed->SampleType = FMP4BoxSampleEntry::ESampleType::QTFFTimecode;
+					break;
+				}
+				case MakeBoxAtom('m','e','t','a'):
+				{
+					Parsed->SampleType = FMP4BoxSampleEntry::ESampleType::TimedMetadata;
+					break;
+				}
+				case MakeBoxAtom('s','b','t','l'):
+				case MakeBoxAtom('s','u','b','t'):
+				{
+					Parsed->SampleType = FMP4BoxSampleEntry::ESampleType::Subtitles;
+					break;
+				}
+				default:
+				{
+					Parsed->SampleType = FMP4BoxSampleEntry::ESampleType::Unsupported;
+					break;
+				}
+			}
+			// Unsupported sample type. No need to continue parsing.
+			if (Parsed->SampleType == FMP4BoxSampleEntry::ESampleType::Unsupported)
+			{
+				check(!"TODO");
+				return;
+			}
+			uint32 EntryCount;
+			ar.Read(EntryCount);
+			for(uint32 ne=0; ne<EntryCount; ++ne)
+			{
+				FMP4BoxInfo bi;
+				if (!ar.ParseIntoBoxInfo(bi, BoxInfo.Offset + BoxInfo.DataOffset + ar.GetCurrentOffset()))
+				{
+					UE_LOG(LogElectraBase, Error, TEXT("Failed to parse `stsd` box!"));
+					return;
+				}
+				TSharedPtr<FMP4BoxBase> SampleEntry;
+				switch(Parsed->SampleType)
+				{
+					case FMP4BoxSampleEntry::ESampleType::Video:
+					{
+						SampleEntry = FMP4BoxVisualSampleEntry::Create(AsWeak(), bi);
+						break;
+					}
+					case FMP4BoxSampleEntry::ESampleType::Audio:
+					{
+						SampleEntry = FMP4BoxAudioSampleEntry::Create(AsWeak(), bi);
+						break;
+					}
+					case FMP4BoxSampleEntry::ESampleType::QTFFTimecode:
+					{
+						SampleEntry = FMP4BoxQTFFTimecodeSampleEntry::Create(AsWeak(), bi);
+						break;
+					}
+					case FMP4BoxSampleEntry::ESampleType::Subtitles:
+					{
+						switch(bi.Type)
+						{
+							case MakeBoxAtom('t','x','3','g'):
+							{
+								SampleEntry = FMP4BoxTX3GSampleEntry::Create(AsWeak(), bi);
+								break;
+							}
+							default:
+							{
+								//check(!"TODO");
+								// Ignored for now.
+								break;
+							}
+						}
+						break;
+					}
+					case FMP4BoxSampleEntry::ESampleType::TimedMetadata:
+					{
+						// Ignored for now.
+						break;
+					}
+				}
+				AddChildBox(MoveTemp(SampleEntry));
+			}
+		}
+
+		uint8 FMP4BoxSTSD::GetBoxVersion()
+		{
+			ParseIfRequired();
+			return Parsed->Version;
+		}
+
+		FMP4BoxSampleEntry::ESampleType FMP4BoxSTSD::GetSampleType()
+		{
+			ParseIfRequired();
+			return Parsed->SampleType;
+		}
+
+/****************************************************************************************************************************************************/
+
+		struct FMP4BoxSTSS::FParsed
+		{
+			TArray<uint32> Entries;
+			uint32 Flags = 0;
+			uint8 Version = 0;
+		};
+		void FMP4BoxSTSS::ParseIfRequired()
+		{
+			if (Parsed.IsValid())
+			{
+				return;
+			}
+			Parsed = MakePimpl<FParsed>();
+			FMP4AtomReaderHelper ar(BoxInfo.Data);
+			ar.ReadVersionAndFlags(Parsed->Version, Parsed->Flags);
+			uint32 NumEntries;
+			ar.Read(NumEntries);
+			Parsed->Entries.AddUninitialized(NumEntries);
+			for(uint32 i=0; i<NumEntries; ++i)
+			{
+				ar.Read(Parsed->Entries[i]);
+			}
+		}
+
+		const TArray<uint32>& FMP4BoxSTSS::GetEntries()
+		{
+			ParseIfRequired();
+			return Parsed->Entries;
+		}
+
+/****************************************************************************************************************************************************/
+
+		struct FMP4BoxSDTP::FParsed
+		{
+			TConstArrayView<uint8> Entries;
+			uint32 Flags = 0;
+			uint8 Version = 0;
+		};
+		void FMP4BoxSDTP::ParseIfRequired()
+		{
+			if (Parsed.IsValid())
+			{
+				return;
+			}
+			Parsed = MakePimpl<FParsed>();
+			FMP4AtomReaderHelper ar(BoxInfo.Data);
+			ar.ReadVersionAndFlags(Parsed->Version, Parsed->Flags);
+			uint32 NumEntries;
+			ar.Read(NumEntries);
+			Parsed->Entries = MakeConstArrayView(ar.GetCurrentDataPointer(), NumEntries);
+		}
+
+		TConstArrayView<uint8> FMP4BoxSDTP::GetEntries()
+		{
+			ParseIfRequired();
+			return Parsed->Entries;
+		}
+
+/****************************************************************************************************************************************************/
+
+		struct FMP4BoxVisualSampleEntry::FParsed
+		{
+			uint8 CompressorName[32] {};
+			uint16 DataReferenceIndex = 0;
+			uint16 Width = 0;
+			uint16 Height = 0;
+			uint16 FrameCount = 0;
+			uint16 Depth = 0;
+		};
+		void FMP4BoxVisualSampleEntry::ParseIfRequired()
+		{
+			if (Parsed.IsValid())
+			{
+				return;
+			}
+			Parsed = MakePimpl<FParsed>();
+			FMP4AtomReaderHelper ar(BoxInfo.Data);
+			// Members of the generic `SampleEntry` class
+			ar.SkipBytes(6);
+			ar.Read(Parsed->DataReferenceIndex);
+			// VisualSampleEntry follows
+			ar.SkipBytes(16);		// `pre_defined`, `reserved`, `pre_defined`
+			ar.Read(Parsed->Width);
+			ar.Read(Parsed->Height);
+			ar.SkipBytes(12);		// `horizresolution`, `vertresolution`, `reserved`
+			ar.Read(Parsed->FrameCount);
+			ar.ReadBytes(Parsed->CompressorName, 32);
+			ar.Read(Parsed->Depth);
+			ar.SkipBytes(2);		// `pre_defined`
+			// There can now be additional boxes following, most notably `pasp` and `clap`.
+			ProcessBoxChildrenRecursively(ar, BoxInfo);
+		}
+
+		uint16 FMP4BoxVisualSampleEntry::GetDataReferenceIndex()
+		{
+			ParseIfRequired();
+			return Parsed->DataReferenceIndex;
+		}
+		uint16 FMP4BoxVisualSampleEntry::GetWidth()
+		{
+			ParseIfRequired();
+			return Parsed->Width;
+		}
+		uint16 FMP4BoxVisualSampleEntry::GetHeight()
+		{
+			ParseIfRequired();
+			return Parsed->Height;
+		}
+		uint16 FMP4BoxVisualSampleEntry::GetFrameCount()
+		{
+			ParseIfRequired();
+			return Parsed->FrameCount;
+		}
+		uint16 FMP4BoxVisualSampleEntry::GetDepth()
+		{
+			ParseIfRequired();
+			return Parsed->Depth;
+		}
+
+/****************************************************************************************************************************************************/
+
+		struct FMP4BoxAudioSampleEntry::FParsed
+		{
+			uint32 SampleRate = 0;
+			uint16 DataReferenceIndex = 0;
+			uint16 Version = 0;
+			uint16 ChannelCount = 0;
+			uint16 SampleSize = 0;
+
+			// QuickTime compatibility
+			uint32 ConstBytesPerAudioPacket = 0;
+			uint32 ConstLPCMFramesPerAudioPacket = 0;
+			int32 FormatSpecificFlags = 0;
+		};
+		void FMP4BoxAudioSampleEntry::ParseIfRequired()
+		{
+			if (Parsed.IsValid())
+			{
+				return;
+			}
+			Parsed = MakePimpl<FParsed>();
+			FMP4AtomReaderHelper ar(BoxInfo.Data);
+			// Members of the generic `SampleEntry` class
+			ar.SkipBytes(6);
+			ar.Read(Parsed->DataReferenceIndex);
+			// AudioSampleEntry follows
+			ar.Read(Parsed->Version);
+			ar.SkipBytes(6);		// `reserved`
+
+			// The audio sample entry can be version 0 or version 1.
+			// In ISO/IEC 14496-12 a version 1 sample is required to be inside a version 1 'stsd' box while in QuickTime a
+			// version 1 sample was allowed in a version 0 'stsd'. The version 1 sample is not identical between ISO and QT
+			// so we do some special handling based on version numbers.
+			if (Parsed->Version == 0 || Parsed->Version == 1)
+			{
+				ar.Read(Parsed->ChannelCount);
+				ar.Read(Parsed->SampleSize);
+				ar.SkipBytes(4);	// `pre_defined`, `reserved`
+				ar.Read(Parsed->SampleRate);
+				// Handle a special case of version 1 QuickTime box. In ISO/IEC 14496-12 the SampleRate must be set to 0x00010000 (1 Hz).
+				// If this is not the case we assume this to be a QuickTime box which adds 4 additional fields *before* any optional boxes.
+				if (Parsed->Version == 1 && Parsed->SampleRate != (1U << 16))
+				{
+					// Which version is the enclosing `stsd` box?
+					auto STSDBox = FindParentBox<FMP4BoxSTSD>(MakeBoxAtom('s','t','s','d'));
+					if (!STSDBox.IsValid())
+					{
+						UE_LOG(LogElectraBase, Error, TEXT("Could not find the parent `stsd` box to parse this AudioSampleEntry!"));
+						return;
+					}
+					if (STSDBox->GetBoxVersion() == 0)
+					{
+						// Assume QuickTime and skip the following elements
+						ar.SkipBytes(4 * 4);	// Samples per packet; Bytes per packet; Bytes per frame; Bytes per Sample
+					}
+				}
+				// The sample rate is stored in the upper 16 bits only. The lower 16 bits are 0.
+				Parsed->SampleRate >>= 16;
+			}
+			// As defined by: https://developer.apple.com/documentation/quicktime-file-format/sound_sample_description_version_2
+			else if (Parsed->Version == 2)
+			{
+				int32 SizeOfStructOnly = 0;
+				int32 NumChannels = 0;
+				int32 ConstBitsPerChannel = 0;
+				uint32 Value32 = 0;
+				uint16 Value16 = 0;
+				union Int64toDouble
+				{
+					uint64 ui64;
+					double dbl64;
+				};
+				Int64toDouble SampleRateCnv;
+				SampleRateCnv.ui64 = 0;
+				ar.Read(Value16);					// always 3
+				ar.Read(Value16);					// always 16
+				ar.Read(Value16);					// always 0xfffe
+				ar.Read(Value16);					// always 0
+				ar.Read(Value32);					// always 65536
+				ar.Read(SizeOfStructOnly);
+				ar.Read(SampleRateCnv.ui64);
+				ar.Read(NumChannels);
+				ar.Read(Value32);					// always 0x7f000000
+				ar.Read(ConstBitsPerChannel);
+
+				ar.Read(Parsed->FormatSpecificFlags);
+				ar.Read(Parsed->ConstBytesPerAudioPacket);
+				ar.Read(Parsed->ConstLPCMFramesPerAudioPacket);
+
+				Parsed->SampleRate = (uint32) SampleRateCnv.dbl64;
+				Parsed->ChannelCount = (uint16) NumChannels;
+				Parsed->SampleSize = (uint16) ConstBitsPerChannel;
+			}
+
+			// There can now be additional boxes following, like `esds`, `chnl`, `dmix`, `btrt`, etc.
+			ProcessBoxChildrenRecursively(ar, BoxInfo);
+		}
+
+		uint16 FMP4BoxAudioSampleEntry::GetDataReferenceIndex()
+		{
+			ParseIfRequired();
+			return Parsed->DataReferenceIndex;
+		}
+		uint32 FMP4BoxAudioSampleEntry::GetSampleRate()
+		{
+			ParseIfRequired();
+			return Parsed->SampleRate;
+		}
+		int32 FMP4BoxAudioSampleEntry::GetChannelCount()
+		{
+			ParseIfRequired();
+			return Parsed->ChannelCount;
+		}
+		int32 FMP4BoxAudioSampleEntry::GetSampleSize()
+		{
+			ParseIfRequired();
+			return Parsed->SampleSize;
+		}
+		bool FMP4BoxAudioSampleEntry::HaveFormatSpecificFlags()
+		{
+			ParseIfRequired();
+			return Parsed->Version == 2;
+		}
+		int32 FMP4BoxAudioSampleEntry::GetFormatSpecificFlags()
+		{
+			ParseIfRequired();
+			return Parsed->FormatSpecificFlags;
+		}
+		uint32 FMP4BoxAudioSampleEntry::GetConstBytesPerAudioPacket()
+		{
+			ParseIfRequired();
+			return Parsed->ConstBytesPerAudioPacket;
+		}
+		uint32 FMP4BoxAudioSampleEntry::GetConstLPCMFramesPerAudioPacket()
+		{
+			ParseIfRequired();
+			return Parsed->ConstLPCMFramesPerAudioPacket;
+		}
+
+/****************************************************************************************************************************************************/
+
+		struct FMP4BoxQTFFTimecodeSampleEntry::FParsed
+		{
+			uint32 Flags = 0;
+			uint32 Timescale = 0;
+			uint32 FrameDuration = 0;
+			uint16 DataReferenceIndex = 0;
+			uint8 NumberOfFrames = 0;
+		};
+		void FMP4BoxQTFFTimecodeSampleEntry::ParseIfRequired()
+		{
+			if (Parsed.IsValid())
+			{
+				return;
+			}
+			Parsed = MakePimpl<FParsed>();
+			FMP4AtomReaderHelper ar(BoxInfo.Data);
+			// Members of the generic `SampleEntry` class
+			ar.SkipBytes(6);
+			ar.Read(Parsed->DataReferenceIndex);
+			// QTFFTimecodeSampleEntry follows
+			ar.SkipBytes(4);		// `reserved`
+			ar.Read(Parsed->Flags);
+			ar.Read(Parsed->Timescale);
+			ar.Read(Parsed->FrameDuration);
+			ar.Read(Parsed->NumberOfFrames);
+			ar.SkipBytes(1);		// `reserved`
+			// There can now be additional boxes following.
+			ProcessBoxChildrenRecursively(ar, BoxInfo);
+		}
+
+		uint16 FMP4BoxQTFFTimecodeSampleEntry::GetDataReferenceIndex()
+		{
+			ParseIfRequired();
+			return Parsed->DataReferenceIndex;
+		}
+		uint32 FMP4BoxQTFFTimecodeSampleEntry::GetFlags()
+		{
+			ParseIfRequired();
+			return Parsed->Flags;
+		}
+		uint32 FMP4BoxQTFFTimecodeSampleEntry::GetTimescale()
+		{
+			ParseIfRequired();
+			return Parsed->Timescale;
+		}
+		uint32 FMP4BoxQTFFTimecodeSampleEntry::GetFrameDuration()
+		{
+			ParseIfRequired();
+			return Parsed->FrameDuration;
+		}
+		uint32 FMP4BoxQTFFTimecodeSampleEntry::GetNumberOfFrames()
+		{
+			ParseIfRequired();
+			return Parsed->NumberOfFrames;
+		}
+
+/****************************************************************************************************************************************************/
+
+		struct FMP4BoxTX3GSampleEntry::FParsed
+		{
+			uint16 DataReferenceIndex = 0;
+		};
+		void FMP4BoxTX3GSampleEntry::ParseIfRequired()
+		{
+			if (Parsed.IsValid())
+			{
+				return;
+			}
+			Parsed = MakePimpl<FParsed>();
+			FMP4AtomReaderHelper ar(BoxInfo.Data);
+			// Members of the generic `SampleEntry` class
+			ar.SkipBytes(6);
+			ar.Read(Parsed->DataReferenceIndex);
+/*
+			// TX3GSampleEntry follows
+			ar.SkipBytes(4);		// `reserved`
+			ar.Read(Parsed->Flags);
+			ar.Read(Parsed->Timescale);
+			ar.Read(Parsed->FrameDuration);
+			ar.Read(Parsed->NumberOfFrames);
+			ar.SkipBytes(1);		// `reserved`
+			// There can now be additional boxes following.
+			ProcessBoxChildrenRecursively(ar, BoxInfo);
+*/
+		}
+
+		uint16 FMP4BoxTX3GSampleEntry::GetDataReferenceIndex()
+		{
+			ParseIfRequired();
+			return Parsed->DataReferenceIndex;
+		}
+
+/****************************************************************************************************************************************************/
+
+		struct FMP4BoxAVCC::FParsed
+		{
+			TConstArrayView<uint8> AVCDecoderConfigurationRecord;
+		};
+		void FMP4BoxAVCC::ParseIfRequired()
+		{
+			if (Parsed.IsValid())
+			{
+				return;
+			}
+			Parsed = MakePimpl<FParsed>();
+			FMP4AtomReaderHelper ar(BoxInfo.Data);
+			Parsed->AVCDecoderConfigurationRecord = MakeConstArrayView(ar.GetCurrentDataPointer(), ar.GetNumBytesRemaining());
+		}
+
+		TConstArrayView<uint8> FMP4BoxAVCC::GetAVCDecoderConfigurationRecord()
+		{
+			ParseIfRequired();
+			return Parsed->AVCDecoderConfigurationRecord;
+		}
+
+/****************************************************************************************************************************************************/
+
+		struct FMP4BoxHVCC::FParsed
+		{
+			TConstArrayView<uint8> HEVCDecoderConfigurationRecord;
+		};
+		void FMP4BoxHVCC::ParseIfRequired()
+		{
+			if (Parsed.IsValid())
+			{
+				return;
+			}
+			Parsed = MakePimpl<FParsed>();
+			FMP4AtomReaderHelper ar(BoxInfo.Data);
+			Parsed->HEVCDecoderConfigurationRecord = MakeConstArrayView(ar.GetCurrentDataPointer(), ar.GetNumBytesRemaining());
+		}
+
+		TConstArrayView<uint8> FMP4BoxHVCC::GetHEVCDecoderConfigurationRecord()
+		{
+			ParseIfRequired();
+			return Parsed->HEVCDecoderConfigurationRecord;
+		}
+
+/****************************************************************************************************************************************************/
+
+		struct FMP4BoxDVCC::FParsed
+		{
+			TConstArrayView<uint8> DOVIDecoderConfigurationRecord;
+		};
+		void FMP4BoxDVCC::ParseIfRequired()
+		{
+			if (Parsed.IsValid())
+			{
+				return;
+			}
+			Parsed = MakePimpl<FParsed>();
+			FMP4AtomReaderHelper ar(BoxInfo.Data);
+			Parsed->DOVIDecoderConfigurationRecord = MakeConstArrayView(ar.GetCurrentDataPointer(), ar.GetNumBytesRemaining());
+		}
+
+		TConstArrayView<uint8> FMP4BoxDVCC::GetDOVIDecoderConfigurationRecord()
+		{
+			ParseIfRequired();
+			return Parsed->DOVIDecoderConfigurationRecord;
+		}
+
+/****************************************************************************************************************************************************/
+
+		struct FMP4BoxDAC3::FParsed
+		{
+			TConstArrayView<uint8> AC3SpecificBox;
+		};
+		void FMP4BoxDAC3::ParseIfRequired()
+		{
+			if (Parsed.IsValid())
+			{
+				return;
+			}
+			Parsed = MakePimpl<FParsed>();
+			FMP4AtomReaderHelper ar(BoxInfo.Data);
+			Parsed->AC3SpecificBox = MakeConstArrayView(ar.GetCurrentDataPointer(), ar.GetNumBytesRemaining());
+		}
+
+		TConstArrayView<uint8> FMP4BoxDAC3::GetAC3SpecificBox()
+		{
+			ParseIfRequired();
+			return Parsed->AC3SpecificBox;
+		}
+
+/****************************************************************************************************************************************************/
+
+		struct FMP4BoxDEC3::FParsed
+		{
+			TConstArrayView<uint8> EC3SpecificBox;
+		};
+		void FMP4BoxDEC3::ParseIfRequired()
+		{
+			if (Parsed.IsValid())
+			{
+				return;
+			}
+			Parsed = MakePimpl<FParsed>();
+			FMP4AtomReaderHelper ar(BoxInfo.Data);
+			Parsed->EC3SpecificBox = MakeConstArrayView(ar.GetCurrentDataPointer(), ar.GetNumBytesRemaining());
+		}
+
+		TConstArrayView<uint8> FMP4BoxDEC3::GetEC3SpecificBox()
+		{
+			ParseIfRequired();
+			return Parsed->EC3SpecificBox;
+		}
+
+/****************************************************************************************************************************************************/
+
+		struct FMP4BoxDFLA::FParsed
+		{
+			TConstArrayView<uint8> FLACSpecificBox;
+		};
+		void FMP4BoxDFLA::ParseIfRequired()
+		{
+			if (Parsed.IsValid())
+			{
+				return;
+			}
+			Parsed = MakePimpl<FParsed>();
+			FMP4AtomReaderHelper ar(BoxInfo.Data);
+			Parsed->FLACSpecificBox = MakeConstArrayView(ar.GetCurrentDataPointer(), ar.GetNumBytesRemaining());
+		}
+
+		TConstArrayView<uint8> FMP4BoxDFLA::GetFLACSpecificBox()
+		{
+			ParseIfRequired();
+			return Parsed->FLACSpecificBox;
+		}
+
+/****************************************************************************************************************************************************/
+
+		struct FMP4BoxDOPS::FParsed
+		{
+			TConstArrayView<uint8> OpusSpecificBox;
+		};
+		void FMP4BoxDOPS::ParseIfRequired()
+		{
+			if (Parsed.IsValid())
+			{
+				return;
+			}
+			Parsed = MakePimpl<FParsed>();
+			FMP4AtomReaderHelper ar(BoxInfo.Data);
+			Parsed->OpusSpecificBox = MakeConstArrayView(ar.GetCurrentDataPointer(), ar.GetNumBytesRemaining());
+		}
+
+		TConstArrayView<uint8> FMP4BoxDOPS::GetOpusSpecificBox()
+		{
+			ParseIfRequired();
+			return Parsed->OpusSpecificBox;
+		}
+
+/****************************************************************************************************************************************************/
+
+		struct FMP4BoxVPCC::FParsed
+		{
+			TConstArrayView<uint8> VPCodecConfigurationBox;
+		};
+		void FMP4BoxVPCC::ParseIfRequired()
+		{
+			if (Parsed.IsValid())
+			{
+				return;
+			}
+			Parsed = MakePimpl<FParsed>();
+			FMP4AtomReaderHelper ar(BoxInfo.Data);
+			Parsed->VPCodecConfigurationBox = MakeConstArrayView(ar.GetCurrentDataPointer(), ar.GetNumBytesRemaining());
+		}
+
+		TConstArrayView<uint8> FMP4BoxVPCC::GetVPCodecConfigurationBox()
+		{
+			ParseIfRequired();
+			return Parsed->VPCodecConfigurationBox;
+		}
+
+/****************************************************************************************************************************************************/
+
+		struct FMP4BoxIODS::FParsed
+		{
+			TConstArrayView<uint8> InitialObjectDescriptor;
+		};
+		void FMP4BoxIODS::ParseIfRequired()
+		{
+			if (Parsed.IsValid())
+			{
+				return;
+			}
+			Parsed = MakePimpl<FParsed>();
+			FMP4AtomReaderHelper ar(BoxInfo.Data);
+			Parsed->InitialObjectDescriptor = MakeConstArrayView(ar.GetCurrentDataPointer(), ar.GetNumBytesRemaining());
+		}
+
+		TConstArrayView<uint8> FMP4BoxIODS::GetObjectDescriptor()
+		{
+			ParseIfRequired();
+			return Parsed->InitialObjectDescriptor;
+		}
+
+/****************************************************************************************************************************************************/
+
+		struct FMP4BoxESDS::FParsed
+		{
+			TConstArrayView<uint8> ESDescriptor;
+			uint32 Flags = 0;
+			uint8 Version = 0;
+		};
+		void FMP4BoxESDS::ParseIfRequired()
+		{
+			if (Parsed.IsValid())
+			{
+				return;
+			}
+			Parsed = MakePimpl<FParsed>();
+			FMP4AtomReaderHelper ar(BoxInfo.Data);
+			ar.ReadVersionAndFlags(Parsed->Version, Parsed->Flags);
+			Parsed->ESDescriptor = MakeConstArrayView(ar.GetCurrentDataPointer(), ar.GetNumBytesRemaining());
+		}
+
+		TConstArrayView<uint8> FMP4BoxESDS::GetESDescriptor()
+		{
+			ParseIfRequired();
+			return Parsed->ESDescriptor;
+		}
+
+/****************************************************************************************************************************************************/
+
+		struct FMP4BoxCOLR::FParsed
+		{
+			FColorNCLX ColorNCLX;
+			EColorType ColorType = EColorType::Unsupported;
+		};
+		void FMP4BoxCOLR::ParseIfRequired()
+		{
+			if (Parsed.IsValid())
+			{
+				return;
+			}
+			Parsed = MakePimpl<FParsed>();
+			FMP4AtomReaderHelper ar(BoxInfo.Data);
+			uint32 colour_type;
+			ar.Read(colour_type);
+			if (colour_type == MakeBoxAtom('n','c','l','c') || colour_type == MakeBoxAtom('n','c','l','x'))
+			{
+				Parsed->ColorType = EColorType::nclx;
+				uint16 colour_primaries, transfer_characteristics, matrix_coefficients;
+				ar.Read(colour_primaries);
+				ar.Read(transfer_characteristics);
+				ar.Read(matrix_coefficients);
+				Parsed->ColorNCLX.colour_primaries = colour_primaries;
+				Parsed->ColorNCLX.transfer_characteristics = transfer_characteristics;
+				Parsed->ColorNCLX.matrix_coefficients = matrix_coefficients;
+				if (colour_type == MakeBoxAtom('n','c','l','x'))
+				{
+					uint8 full_range_flag;
+					ar.Read(full_range_flag);
+					Parsed->ColorNCLX.full_range_flag = full_range_flag >> 7;
+				}
+			}
+		}
+
+		FMP4BoxCOLR::EColorType FMP4BoxCOLR::GetColorType()
+		{
+			ParseIfRequired();
+			return Parsed->ColorType;
+		}
+		const FMP4BoxCOLR::FColorNCLX& FMP4BoxCOLR::GetColorNCLX()
+		{
+			ParseIfRequired();
+			return Parsed->ColorNCLX;
+		}
+
+/****************************************************************************************************************************************************/
+
+		struct FMP4BoxBTRT::FParsed
+		{
+			uint32 BufferSizeDB = 0;
+			uint32 MaxBitrate = 0;
+			uint32 AverageBitrate = 0;
+		};
+		void FMP4BoxBTRT::ParseIfRequired()
+		{
+			if (Parsed.IsValid())
+			{
+				return;
+			}
+			Parsed = MakePimpl<FParsed>();
+			FMP4AtomReaderHelper ar(BoxInfo.Data);
+			ar.Read(Parsed->BufferSizeDB);
+			ar.Read(Parsed->MaxBitrate);
+			ar.Read(Parsed->AverageBitrate);
+		}
+
+		uint32 FMP4BoxBTRT::GetBufferSizeDB()
+		{
+			ParseIfRequired();
+			return Parsed->BufferSizeDB;
+		}
+		uint32 FMP4BoxBTRT::GetMaxBitrate()
+		{
+			ParseIfRequired();
+			return Parsed->MaxBitrate;
+		}
+		uint32 FMP4BoxBTRT::GetAverageBitrate()
+		{
+			ParseIfRequired();
+			return Parsed->AverageBitrate;
+		}
+
+/****************************************************************************************************************************************************/
+
+		struct FMP4BoxSTTS::FParsed
+		{
+			TArray<FEntry> Entries;
+			int64 TotalDuration = 0;
+			uint32 NumTotalSamples = 0;
+			uint32 Flags = 0;
+			uint8 Version = 0;
+		};
+		void FMP4BoxSTTS::ParseIfRequired()
+		{
+			if (Parsed.IsValid())
+			{
+				return;
+			}
+			Parsed = MakePimpl<FParsed>();
+			FMP4AtomReaderHelper ar(BoxInfo.Data);
+			ar.ReadVersionAndFlags(Parsed->Version, Parsed->Flags);
+			uint32 entry_count;
+			ar.Read(entry_count);
+			Parsed->Entries.SetNumUninitialized(entry_count);
+			for(uint32 i=0; i<entry_count; ++i)
+			{
+				ar.Read(Parsed->Entries[i].sample_count);
+				ar.Read(Parsed->Entries[i].sample_delta);
+				Parsed->NumTotalSamples += Parsed->Entries[i].sample_count;
+				Parsed->TotalDuration += Parsed->Entries[i].sample_count * Parsed->Entries[i].sample_delta;
+			}
+		}
+
+		uint32 FMP4BoxSTTS::GetNumTotalSamples()
+		{
+			ParseIfRequired();
+			return Parsed->NumTotalSamples;
+		}
+		int64 FMP4BoxSTTS::GetTotalDuration()
+		{
+			ParseIfRequired();
+			return Parsed->TotalDuration;
+		}
+		const TArray<FMP4BoxSTTS::FEntry>& FMP4BoxSTTS::GetEntries()
+		{
+			ParseIfRequired();
+			return Parsed->Entries;
+		}
+
+/****************************************************************************************************************************************************/
+
+		struct FMP4BoxCTTS::FParsed
+		{
+			TArray<FEntry> Entries;
+			uint32 NumTotalSamples = 0;
+			uint32 Flags = 0;
+			uint8 Version = 0;
+		};
+		void FMP4BoxCTTS::ParseIfRequired()
+		{
+			if (Parsed.IsValid())
+			{
+				return;
+			}
+			Parsed = MakePimpl<FParsed>();
+			FMP4AtomReaderHelper ar(BoxInfo.Data);
+			ar.ReadVersionAndFlags(Parsed->Version, Parsed->Flags);
+			uint32 entry_count;
+			ar.Read(entry_count);
+			if (Parsed->Version == 0)
+			{
+				Parsed->Entries.SetNumUninitialized(entry_count);
+				uint32 Value32;
+				for(uint32 i=0; i<entry_count; ++i)
+				{
+					ar.Read(Parsed->Entries[i].sample_count);
+					Parsed->NumTotalSamples += Parsed->Entries[i].sample_count;
+					ar.Read(Value32);
+					Parsed->Entries[i].sample_offset = Value32;
+				}
+			}
+			else if (Parsed->Version == 1)
+			{
+				Parsed->Entries.SetNumUninitialized(entry_count);
+				int32 Value32;
+				for(uint32 i=0; i<entry_count; ++i)
+				{
+					ar.Read(Parsed->Entries[i].sample_count);
+					Parsed->NumTotalSamples += Parsed->Entries[i].sample_count;
+					ar.Read(Value32);
+					Parsed->Entries[i].sample_offset = Value32;
+				}
+			}
+		}
+
+		uint8 FMP4BoxCTTS::GetBoxVersion()
+		{
+			ParseIfRequired();
+			return Parsed->Version;
+		}
+		uint32 FMP4BoxCTTS::GetNumTotalSamples()
+		{
+			ParseIfRequired();
+			return Parsed->NumTotalSamples;
+		}
+		const TArray<FMP4BoxCTTS::FEntry>& FMP4BoxCTTS::GetEntries()
+		{
+			ParseIfRequired();
+			return Parsed->Entries;
+		}
+
+/****************************************************************************************************************************************************/
+
+		struct FMP4BoxSTSC::FParsed
+		{
+			TArray<FEntry> Entries;
+			uint32 Flags = 0;
+			uint8 Version = 0;
+		};
+		void FMP4BoxSTSC::ParseIfRequired()
+		{
+			if (Parsed.IsValid())
+			{
+				return;
+			}
+			Parsed = MakePimpl<FParsed>();
+			FMP4AtomReaderHelper ar(BoxInfo.Data);
+			ar.ReadVersionAndFlags(Parsed->Version, Parsed->Flags);
+			uint32 entry_count;
+			ar.Read(entry_count);
+			Parsed->Entries.SetNumUninitialized(entry_count);
+			static_assert(sizeof(FEntry::first_chunk) == 4);
+			static_assert(sizeof(FEntry::samples_per_chunk) == 4);
+			static_assert(sizeof(FEntry::sample_description_index) == 4);
+			for(uint32 i=0; i<entry_count; ++i)
+			{
+				ar.Read(Parsed->Entries[i].first_chunk);
+				ar.Read(Parsed->Entries[i].samples_per_chunk);
+				ar.Read(Parsed->Entries[i].sample_description_index);
+			}
+		}
+
+		const TArray<FMP4BoxSTSC::FEntry>& FMP4BoxSTSC::GetEntries()
+		{
+			ParseIfRequired();
+			return Parsed->Entries;
+		}
+
+/****************************************************************************************************************************************************/
+
+		struct FMP4BoxSTSZ::FParsed
+		{
+			const uint32* SizeArray32 = nullptr;
+			uint32 NumberOfSamples = 0;
+			uint32 ConstantSampleSize = 0;
+			uint32 Flags = 0;
+			uint8 Version = 0;
+		};
+		void FMP4BoxSTSZ::ParseIfRequired()
+		{
+			if (Parsed.IsValid())
+			{
+				return;
+			}
+			Parsed = MakePimpl<FParsed>();
+			FMP4AtomReaderHelper ar(BoxInfo.Data);
+			ar.ReadVersionAndFlags(Parsed->Version, Parsed->Flags);
+			ar.Read(Parsed->ConstantSampleSize);
+			ar.Read(Parsed->NumberOfSamples);
+			Parsed->SizeArray32 = (const uint32*) ar.GetCurrentDataPointer();
+		}
+
+		uint32 FMP4BoxSTSZ::GetNumberOfSamples()
+		{
+			ParseIfRequired();
+			return Parsed->NumberOfSamples;
+		}
+		uint32 FMP4BoxSTSZ::GetSizeOfSample(uint32 InIndex)
+		{
+			ParseIfRequired();
+			return Parsed->ConstantSampleSize ? Parsed->ConstantSampleSize : InIndex < Parsed->NumberOfSamples ? GetFromBigEndian(Parsed->SizeArray32[InIndex]) : 0;
+		}
+
+/****************************************************************************************************************************************************/
+
+		struct FMP4BoxSTCO::FParsed
+		{
+			const uint32* SizeArray32 = nullptr;
+			const uint64* SizeArray64 = nullptr;
+			uint32 NumberOfEntries = 0;
+			uint32 Flags = 0;
+			uint8 Version = 0;
+		};
+		void FMP4BoxSTCO::ParseIfRequired()
+		{
+			if (Parsed.IsValid())
+			{
+				return;
+			}
+			Parsed = MakePimpl<FParsed>();
+			FMP4AtomReaderHelper ar(BoxInfo.Data);
+			ar.ReadVersionAndFlags(Parsed->Version, Parsed->Flags);
+			ar.Read(Parsed->NumberOfEntries);
+			if (BoxInfo.Type == MakeBoxAtom('c','o','6','4'))
+			{
+				Parsed->SizeArray64 = (const uint64*) ar.GetCurrentDataPointer();
+			}
+			else
+			{
+				Parsed->SizeArray32 = (const uint32*) ar.GetCurrentDataPointer();
+			}
+		}
+
+		uint32 FMP4BoxSTCO::GetNumberOfEntries()
+		{
+			ParseIfRequired();
+			return Parsed->NumberOfEntries;
+		}
+		uint64 FMP4BoxSTCO::GetChunkOffset(uint32 InIndex)
+		{
+			ParseIfRequired();
+			return InIndex < Parsed->NumberOfEntries ? Parsed->SizeArray32 ? GetFromBigEndian(Parsed->SizeArray32[InIndex]) : GetFromBigEndian(Parsed->SizeArray64[InIndex]) : 0;
+		}
+
+/****************************************************************************************************************************************************/
+
+		struct FMP4BoxSAIZ::FParsed
+		{
+			TConstArrayView<uint8> SampleInfoSizes;
+			uint32 AuxInfoType = 0;
+			uint32 AuxInfoTypeParameter = 0;
+			uint32 SampleCount = 0;
+			uint32 Flags = 0;
+			uint8 Version = 0;
+			uint8 DefaultSampleInfoSize = 0;
+		};
+		void FMP4BoxSAIZ::ParseIfRequired()
+		{
+			if (Parsed.IsValid())
+			{
+				return;
+			}
+			Parsed = MakePimpl<FParsed>();
+			FMP4AtomReaderHelper ar(BoxInfo.Data);
+			ar.ReadVersionAndFlags(Parsed->Version, Parsed->Flags);
+			if ((Parsed->Flags & 1) != 0)
+			{
+				ar.Read(Parsed->AuxInfoType);
+				ar.Read(Parsed->AuxInfoTypeParameter);
+			}
+			ar.Read(Parsed->DefaultSampleInfoSize);
+			ar.Read(Parsed->SampleCount);
+			if (Parsed->DefaultSampleInfoSize == 0 && Parsed->SampleCount != 0)
+			{
+				Parsed->SampleInfoSizes = MakeConstArrayView(ar.GetCurrentDataPointer(), Parsed->SampleCount);
+			}
+		}
+
+		void FMP4BoxSAIZ::Test()
+		{
+			ParseIfRequired();
+		}
+
+/****************************************************************************************************************************************************/
+
+		struct FMP4BoxSAIO::FParsed
+		{
+			const uint64* Offsets64 = nullptr;
+			const uint32* Offsets32 = nullptr;
+			uint32 AuxInfoType = 0;
+			uint32 AuxInfoTypeParameter = 0;
+			uint32 EntryCount = 0;
+			uint32 Flags = 0;
+			uint8 Version = 0;
+		};
+		void FMP4BoxSAIO::ParseIfRequired()
+		{
+			if (Parsed.IsValid())
+			{
+				return;
+			}
+			Parsed = MakePimpl<FParsed>();
+			FMP4AtomReaderHelper ar(BoxInfo.Data);
+			ar.ReadVersionAndFlags(Parsed->Version, Parsed->Flags);
+			if ((Parsed->Flags & 1) != 0)
+			{
+				ar.Read(Parsed->AuxInfoType);
+				ar.Read(Parsed->AuxInfoTypeParameter);
+			}
+			ar.Read(Parsed->EntryCount);
+			// Remember a 32 or 64 bit pointer to where the offsets are.
+			// When accessing them remember that they are in big endian and need to be swapped on return value!
+			if (Parsed->Version == 0)
+			{
+				Parsed->Offsets32 = (const uint32*)ar.GetCurrentDataPointer();
+			}
+			else
+			{
+				Parsed->Offsets64 = (const uint64*)ar.GetCurrentDataPointer();
+			}
+		}
+
+		void FMP4BoxSAIO::Test()
+		{
+			ParseIfRequired();
+		}
+
+/****************************************************************************************************************************************************/
+
+		struct FMP4BoxSGPD::FParsed
+		{
+			TArray<TConstArrayView<uint8>> GroupDescriptionEntries;
+			uint32 GroupingType = 0;
+			uint32 DefaultLength = 0;
+			uint32 DefaultGroupDescriptionIndex = 0;
+			uint32 Flags = 0;
+			uint8 Version = 0;
+		#if !UE_BUILD_SHIPPING
+			char Name[5]{0};
+		#endif
+		};
+		void FMP4BoxSGPD::ParseIfRequired()
+		{
+			if (Parsed.IsValid())
+			{
+				return;
+			}
+			Parsed = MakePimpl<FParsed>();
+			FMP4AtomReaderHelper ar(BoxInfo.Data);
+			ar.ReadVersionAndFlags(Parsed->Version, Parsed->Flags);
+			if (Parsed->Version == 0)
+			{
+				UE_LOG(LogElectraBase, Warning, TEXT("Version 0 of the `sgpd` box is deprecated. Box will be ignored."));
+				return;
+			}
+			ar.Read(Parsed->GroupingType);
+		#if !UE_BUILD_SHIPPING
+			Parsed->Name[0] = (char) ((Parsed->GroupingType >> 24) & 255);
+			Parsed->Name[1] = (char) ((Parsed->GroupingType >> 16) & 255);
+			Parsed->Name[2] = (char) ((Parsed->GroupingType >>  8) & 255);
+			Parsed->Name[3] = (char) ((Parsed->GroupingType >>  0) & 255);
+		#endif
+			if (Parsed->Version >= 1)
+			{
+				ar.Read(Parsed->DefaultLength);
+			}
+			if (Parsed->Version >= 2)
+			{
+				ar.Read(Parsed->DefaultGroupDescriptionIndex);
+			}
+			uint32 entry_count;
+			ar.Read(entry_count);
+			Parsed->GroupDescriptionEntries.SetNum(entry_count);
+			for(uint32 i=0; i<entry_count; ++i)
+			{
+				uint32 description_length = Parsed->DefaultLength;
+				if (Parsed->Version == 1 && description_length == 0)
+				{
+					ar.Read(description_length);
+				}
+				Parsed->GroupDescriptionEntries[i] = MakeConstArrayView(ar.GetCurrentDataPointer(), description_length);
+			}
+		}
+
+		uint32 FMP4BoxSGPD::GetGroupingType()
+		{
+			ParseIfRequired();
+			return Parsed->GroupingType;
+		}
+		const TArray<TConstArrayView<uint8>>& FMP4BoxSGPD::GetGroupDescriptionEntries()
+		{
+			ParseIfRequired();
+			return Parsed->GroupDescriptionEntries;
+		}
+		uint32 FMP4BoxSGPD::GetDefaultGroupDescriptionIndex()
+		{
+			ParseIfRequired();
+			return Parsed->DefaultGroupDescriptionIndex;
+		}
+
+/****************************************************************************************************************************************************/
+
+		struct FMP4BoxSBGP::FParsed
+		{
+			TArray<FEntry> Entries;
+			uint32 GroupingType = 0;
+			uint32 GroupingTypeParameter = 0;
+			uint32 NumTotalSamples = 0;
+			uint32 Flags = 0;
+			uint8 Version = 0;
+		#if !UE_BUILD_SHIPPING
+			char Name[5]{0};
+		#endif
+		};
+		void FMP4BoxSBGP::ParseIfRequired()
+		{
+			if (Parsed.IsValid())
+			{
+				return;
+			}
+			Parsed = MakePimpl<FParsed>();
+			FMP4AtomReaderHelper ar(BoxInfo.Data);
+			ar.ReadVersionAndFlags(Parsed->Version, Parsed->Flags);
+			ar.Read(Parsed->GroupingType);
+		#if !UE_BUILD_SHIPPING
+			Parsed->Name[0] = (char) ((Parsed->GroupingType >> 24) & 255);
+			Parsed->Name[1] = (char) ((Parsed->GroupingType >> 16) & 255);
+			Parsed->Name[2] = (char) ((Parsed->GroupingType >>  8) & 255);
+			Parsed->Name[3] = (char) ((Parsed->GroupingType >>  0) & 255);
+		#endif
+			if (Parsed->Version == 1)
+			{
+				ar.Read(Parsed->GroupingTypeParameter);
+			}
+			uint32 entry_count;
+			ar.Read(entry_count);
+			Parsed->Entries.SetNum(entry_count);
+			static_assert(sizeof(FEntry::sample_count) == 4);
+			static_assert(sizeof(FEntry::group_description_index) == 4);
+			for(uint32 i=0; i<entry_count; ++i)
+			{
+				ar.Read(Parsed->Entries[i].sample_count);
+				Parsed->NumTotalSamples += Parsed->Entries[i].sample_count;
+				ar.Read(Parsed->Entries[i].group_description_index);
+			}
+		}
+
+		uint32 FMP4BoxSBGP::GetGroupingType()
+		{
+			ParseIfRequired();
+			return Parsed->GroupingType;
+		}
+		uint32 FMP4BoxSBGP::GetGroupingTypeParameter()
+		{
+			ParseIfRequired();
+			return Parsed->GroupingTypeParameter;
+		}
+		const TArray<FMP4BoxSBGP::FEntry>& FMP4BoxSBGP::GetEntries()
+		{
+			ParseIfRequired();
+			return Parsed->Entries;
+		}
+		uint32 FMP4BoxSBGP::GetNumTotalSamples()
+		{
+			ParseIfRequired();
+			return Parsed->NumTotalSamples;
+		}
+
+/****************************************************************************************************************************************************/
+
+		struct FMP4BoxMEHD::FParsed
+		{
+			uint64 FragmentDuration = 0;
+			uint32 Flags = 0;
+			uint8 Version = 0;
+		};
+		void FMP4BoxMEHD::ParseIfRequired()
+		{
+			if (Parsed.IsValid())
+			{
+				return;
+			}
+			Parsed = MakePimpl<FParsed>();
+			FMP4AtomReaderHelper ar(BoxInfo.Data);
+			ar.ReadVersionAndFlags(Parsed->Version, Parsed->Flags);
+			if (Parsed->Version == 1)
+			{
+				ar.Read(Parsed->FragmentDuration);
+			}
+			else
+			{
+				uint32 Value32;
+				ar.Read(Value32);
+				Parsed->FragmentDuration = Value32;
+			}
+		}
+
+		uint64 FMP4BoxMEHD::GetFragmentDuration()
+		{
+			ParseIfRequired();
+			return Parsed->FragmentDuration;
+		}
+
+/****************************************************************************************************************************************************/
+
+		struct FMP4BoxTREX::FParsed
+		{
+			uint32 TrackID = 0;
+			uint32 DefaultSampleDescriptionIndex = 0;
+			uint32 DefaultSampleDuration = 0;
+			uint32 DefaultSampleSize = 0;
+			uint32 DefaultSampleFlags = 0;
+			uint32 Flags = 0;
+			uint8 Version = 0;
+		};
+		void FMP4BoxTREX::ParseIfRequired()
+		{
+			if (Parsed.IsValid())
+			{
+				return;
+			}
+			Parsed = MakePimpl<FParsed>();
+			FMP4AtomReaderHelper ar(BoxInfo.Data);
+			ar.ReadVersionAndFlags(Parsed->Version, Parsed->Flags);
+			ar.Read(Parsed->TrackID);
+			ar.Read(Parsed->DefaultSampleDescriptionIndex);
+			ar.Read(Parsed->DefaultSampleDuration);
+			ar.Read(Parsed->DefaultSampleSize);
+			ar.Read(Parsed->DefaultSampleFlags);
+		}
+
+		uint32 FMP4BoxTREX::GetTrackID()
+		{
+			ParseIfRequired();
+			return Parsed->TrackID;
+		}
+		uint32 FMP4BoxTREX::GetDefaultSampleDescriptionIndex()
+		{
+			ParseIfRequired();
+			return Parsed->DefaultSampleDescriptionIndex;
+		}
+		uint32 FMP4BoxTREX::GetDefaultSampleDuration()
+		{
+			ParseIfRequired();
+			return Parsed->DefaultSampleDuration;
+		}
+		uint32 FMP4BoxTREX::GetDefaultSampleSize()
+		{
+			ParseIfRequired();
+			return Parsed->DefaultSampleSize;
+		}
+		uint32 FMP4BoxTREX::GetDefaultSampleFlags()
+		{
+			ParseIfRequired();
+			return Parsed->DefaultSampleFlags;
+		}
+
+/****************************************************************************************************************************************************/
+
+		struct FMP4BoxMFHD::FParsed
+		{
+			uint32 SequenceNumber = 0;
+			uint32 Flags = 0;
+			uint8 Version = 0;
+		};
+		void FMP4BoxMFHD::ParseIfRequired()
+		{
+			if (Parsed.IsValid())
+			{
+				return;
+			}
+			Parsed = MakePimpl<FParsed>();
+			FMP4AtomReaderHelper ar(BoxInfo.Data);
+			ar.ReadVersionAndFlags(Parsed->Version, Parsed->Flags);
+			ar.Read(Parsed->SequenceNumber);
+		}
+
+		uint32 FMP4BoxMFHD::GetSequenceNumber()
+		{
+			ParseIfRequired();
+			return Parsed->SequenceNumber;
+		}
+
+/****************************************************************************************************************************************************/
+
+		struct FMP4BoxTFHD::FParsed
+		{
+			uint64 BaseDataOffset = 0;
+			uint32 TrackID = 0;
+			uint32 SampleDescriptionIndex = 0;
+			uint32 DefaultSampleDuration = 0;
+			uint32 DefaultSampleSize = 0;
+			uint32 DefaultSampleFlags = 0;
+			uint32 Flags = 0;
+			uint8 Version = 0;
+		};
+		void FMP4BoxTFHD::ParseIfRequired()
+		{
+			if (Parsed.IsValid())
+			{
+				return;
+			}
+			Parsed = MakePimpl<FParsed>();
+			FMP4AtomReaderHelper ar(BoxInfo.Data);
+			ar.ReadVersionAndFlags(Parsed->Version, Parsed->Flags);
+			ar.Read(Parsed->TrackID);
+			if ((Parsed->Flags & 1) != 0)
+			{
+				ar.Read(Parsed->BaseDataOffset);			// base_data_offset
+			}
+			if ((Parsed->Flags & 2) != 0)
+			{
+				ar.Read(Parsed->SampleDescriptionIndex);	// sample_description_index
+			}
+			if ((Parsed->Flags & 8) != 0)
+			{
+				ar.Read(Parsed->DefaultSampleDuration);		// default_sample_duration
+			}
+			if ((Parsed->Flags & 16) != 0)
+			{
+				ar.Read(Parsed->DefaultSampleSize);			// default_sample_size
+			}
+			if ((Parsed->Flags & 32) != 0)
+			{
+				ar.Read(Parsed->DefaultSampleFlags);		// default_sample_flags
+			}
+		}
+
+		uint32 FMP4BoxTFHD::GetTrackID()
+		{
+			ParseIfRequired();
+			return Parsed->TrackID;
+		}
+
+		bool FMP4BoxTFHD::HasBaseDataOffset()
+		{
+			ParseIfRequired();
+			return (Parsed->Flags & 0x000001) != 0;
+		}
+		uint64 FMP4BoxTFHD::GetBaseDataOffset()
+		{
+			ParseIfRequired();
+			return Parsed->BaseDataOffset;
+		}
+
+		bool FMP4BoxTFHD::HasSampleDescriptionIndex()
+		{
+			ParseIfRequired();
+			return (Parsed->Flags & 0x000002) != 0;
+		}
+		uint32 FMP4BoxTFHD::GetSampleDescriptionIndex()
+		{
+			ParseIfRequired();
+			return Parsed->SampleDescriptionIndex;
+		}
+
+		bool FMP4BoxTFHD::HasDefaultSampleDuration()
+		{
+			ParseIfRequired();
+			return (Parsed->Flags & 0x000008) != 0;
+		}
+		uint32 FMP4BoxTFHD::GetDefaultSampleDuration()
+		{
+			ParseIfRequired();
+			return Parsed->DefaultSampleDuration;
+		}
+
+		bool FMP4BoxTFHD::HasDefaultSampleSize()
+		{
+			ParseIfRequired();
+			return (Parsed->Flags & 0x000010) != 0;
+		}
+		uint32 FMP4BoxTFHD::GetDefaultSampleSize()
+		{
+			ParseIfRequired();
+			return Parsed->DefaultSampleSize;
+		}
+
+		bool FMP4BoxTFHD::HasDefaultSampleFlags()
+		{
+			ParseIfRequired();
+			return (Parsed->Flags & 0x000020) != 0;
+		}
+		uint32 FMP4BoxTFHD::GetDefaultSampleFlags()
+		{
+			ParseIfRequired();
+			return Parsed->DefaultSampleFlags;
+		}
+
+		bool FMP4BoxTFHD::IsDurationEmpty()
+		{
+			ParseIfRequired();
+			return (Parsed->Flags & 0x010000) != 0;
+		}
+
+		bool FMP4BoxTFHD::IsMoofDefaultBase()
+		{
+			ParseIfRequired();
+			return (Parsed->Flags & 0x020000) != 0;
+		}
+
+/****************************************************************************************************************************************************/
+
+		struct FMP4BoxTFDT::FParsed
+		{
+			uint64 BaseMediaDecodeTime = 0;
+			uint32 Flags = 0;
+			uint8 Version = 0;
+		};
+		void FMP4BoxTFDT::ParseIfRequired()
+		{
+			if (Parsed.IsValid())
+			{
+				return;
+			}
+			Parsed = MakePimpl<FParsed>();
+			FMP4AtomReaderHelper ar(BoxInfo.Data);
+			ar.ReadVersionAndFlags(Parsed->Version, Parsed->Flags);
+			if (Parsed->Version == 1)
+			{
+				ar.Read(Parsed->BaseMediaDecodeTime);
+			}
+			else
+			{
+				uint32 Value32;
+				ar.Read(Value32);
+				Parsed->BaseMediaDecodeTime = Value32;
+			}
+		}
+
+		uint64 FMP4BoxTFDT::GetBaseMediaDecodeTime()
+		{
+			ParseIfRequired();
+			return Parsed->BaseMediaDecodeTime;
+		}
+
+/****************************************************************************************************************************************************/
+
+		struct FMP4BoxTRUN::FParsed
+		{
+			TArray<uint32> SampleDurations;
+			TArray<uint32> SampleSizes;
+			TArray<uint32> SampleFlags;
+			TArray<int32> SampleCompositionTimeOffsets;
+			uint32 SampleCount = 0;
+			uint32 FirstSampleFlags = 0;
+			uint32 Flags = 0;
+			int32 DataOffset = 0;
+			uint8 Version = 0;
+		};
+		void FMP4BoxTRUN::ParseIfRequired()
+		{
+			if (Parsed.IsValid())
+			{
+				return;
+			}
+			Parsed = MakePimpl<FParsed>();
+			FMP4AtomReaderHelper ar(BoxInfo.Data);
+			ar.ReadVersionAndFlags(Parsed->Version, Parsed->Flags);
+			ar.Read(Parsed->SampleCount);
+			if ((Parsed->Flags & 1) != 0)
+			{
+				ar.Read(Parsed->DataOffset);
+			}
+			if ((Parsed->Flags & 4) != 0)
+			{
+				ar.Read(Parsed->FirstSampleFlags);
+			}
+			if (Parsed->SampleCount)
+			{
+				if ((Parsed->Flags & 0x100) != 0)
+				{
+					Parsed->SampleDurations.SetNumUninitialized(Parsed->SampleCount);
+				}
+				if ((Parsed->Flags & 0x200) != 0)
+				{
+					Parsed->SampleSizes.SetNumUninitialized(Parsed->SampleCount);
+				}
+				if ((Parsed->Flags & 0x400) != 0)
+				{
+					Parsed->SampleFlags.SetNumUninitialized(Parsed->SampleCount);
+				}
+				if ((Parsed->Flags & 0x800) != 0)
+				{
+					Parsed->SampleCompositionTimeOffsets.SetNumUninitialized(Parsed->SampleCount);
+				}
+				for(uint32 i=0; i<Parsed->SampleCount; ++i)
+				{
+					if ((Parsed->Flags & 0x100) != 0)
+					{
+						ar.Read(Parsed->SampleDurations[i]);
+					}
+					if ((Parsed->Flags & 0x200) != 0)
+					{
+						ar.Read(Parsed->SampleSizes[i]);
+					}
+					if ((Parsed->Flags & 0x400) != 0)
+					{
+						ar.Read(Parsed->SampleFlags[i]);
+					}
+					if ((Parsed->Flags & 0x800) != 0)
+					{
+						if (Parsed->Version == 0)
+						{
+							uint32 Value32;
+							ar.Read(Value32);
+							/*
+								Because we want to handle only signed time offsets we check if the value can actually be
+								presented as such. If not then the first question would be why the value is that large,
+								which could indicate a bad file. If that is legitimate then we would need to change
+								our SampleCompositionTimeOffsets to be an int64 table.
+							*/
+
+							if (Value32 > 0x7fffffffU)
+							{
+								UE_LOG(LogElectraBase, Warning, TEXT("`trun` version 0 time value cannot be represented as a signed value. Why is it so large?"));
+								Value32 = 0x7fffffffU;
+							}
+							Parsed->SampleCompositionTimeOffsets[i] = (int32)Value32;
+						}
+						else
+						{
+							ar.Read(Parsed->SampleCompositionTimeOffsets[i]);
+						}
+					}
+				}
+			}
+		}
+
+		uint32 FMP4BoxTRUN::GetNumberOfSamples()
+		{
+			ParseIfRequired();
+			return Parsed->SampleCount;
+		}
+		bool FMP4BoxTRUN::HasSampleOffset()
+		{
+			ParseIfRequired();
+			return (Parsed->Flags & 0x000001) != 0;
+		}
+		int32 FMP4BoxTRUN::GetSampleOffset()
+		{
+			ParseIfRequired();
+			return Parsed->DataOffset;
+		}
+		bool FMP4BoxTRUN::HasFirstSampleFlags()
+		{
+			ParseIfRequired();
+			return (Parsed->Flags & 0x000004) != 0;
+		}
+		uint32 FMP4BoxTRUN::GetFirstSampleFlags()
+		{
+			ParseIfRequired();
+			return Parsed->FirstSampleFlags;
+		}
+		bool FMP4BoxTRUN::HasSampleDurations()
+		{
+			ParseIfRequired();
+			return (Parsed->Flags & 0x000100) != 0;
+		}
+		const TArray<uint32>& FMP4BoxTRUN::GetSampleDurations()
+		{
+			ParseIfRequired();
+			return Parsed->SampleDurations;
+		}
+		bool FMP4BoxTRUN::HasSampleSizes()
+		{
+			ParseIfRequired();
+			return (Parsed->Flags & 0x000200) != 0;
+		}
+		const TArray<uint32>& FMP4BoxTRUN::GetSampleSizes()
+		{
+			ParseIfRequired();
+			return Parsed->SampleSizes;
+		}
+		bool FMP4BoxTRUN::HasSampleFlags()
+		{
+			ParseIfRequired();
+			return (Parsed->Flags & 0x000400) != 0;
+		}
+		const TArray<uint32>& FMP4BoxTRUN::GetSampleFlags()
+		{
+			ParseIfRequired();
+			return Parsed->SampleFlags;
+		}
+		bool FMP4BoxTRUN::HasSampleCompositionTimeOffsets()
+		{
+			ParseIfRequired();
+			return (Parsed->Flags & 0x000800) != 0;
+		}
+		const TArray<int32>& FMP4BoxTRUN::GetSampleCompositionTimeOffsets()
+		{
+			ParseIfRequired();
+			return Parsed->SampleCompositionTimeOffsets;
+		}
+
+/****************************************************************************************************************************************************/
+
+		struct FMP4BoxPSSH::FParsed
+		{
+			TConstArrayView<uint8> SystemID;
+			TArray<TConstArrayView<uint8>> KIDs;
+			TConstArrayView<uint8> Data;
+			uint32 Flags = 0;
+			uint8 Version = 0;
+		};
+		void FMP4BoxPSSH::ParseIfRequired()
+		{
+			if (Parsed.IsValid())
+			{
+				return;
+			}
+			Parsed = MakePimpl<FParsed>();
+			FMP4AtomReaderHelper ar(BoxInfo.Data);
+			ar.ReadVersionAndFlags(Parsed->Version, Parsed->Flags);
+			Parsed->SystemID = MakeConstArrayView(ar.GetCurrentDataPointer(), 16);
+			ar.SkipBytes(16);
+			if (Parsed->Version > 0)
+			{
+				uint32 KID_count;
+				ar.Read(KID_count);
+				Parsed->KIDs.SetNum(KID_count);
+				for(uint32 i=0; i<KID_count; ++i)
+				{
+					Parsed->KIDs[i] = MakeConstArrayView(ar.GetCurrentDataPointer(), 16);
+				}
+			}
+			uint32 DataSize;
+			ar.Read(DataSize);
+			Parsed->Data = MakeConstArrayView(ar.GetCurrentDataPointer(), DataSize);
+		}
+
+		TConstArrayView<uint8> FMP4BoxPSSH::GetSystemID()
+		{
+			ParseIfRequired();
+			return Parsed->SystemID;
+		}
+		const TArray<TConstArrayView<uint8>>& FMP4BoxPSSH::GetKIDs()
+		{
+			ParseIfRequired();
+			return Parsed->KIDs;
+		}
+		TConstArrayView<uint8> FMP4BoxPSSH::GetData()
+		{
+			ParseIfRequired();
+			return Parsed->Data;
+		}
+
+/****************************************************************************************************************************************************/
+
+		struct FMP4BoxTENC::FParsed
+		{
+			TConstArrayView<uint8> DefaultKID;
+			TConstArrayView<uint8> DefaultConstantIV;
+			uint32 Flags = 0;
+			uint8 Version = 0;
+			uint8 DefaultCryptByteBlock = 0;
+			uint8 DefaultSkipByteBlock = 0;
+			uint8 DefaultIsProtected = 0;
+			uint8 DefaultPerSampleIVSize = 0;
+		};
+		void FMP4BoxTENC::ParseIfRequired()
+		{
+			if (Parsed.IsValid())
+			{
+				return;
+			}
+			Parsed = MakePimpl<FParsed>();
+			FMP4AtomReaderHelper ar(BoxInfo.Data);
+			ar.ReadVersionAndFlags(Parsed->Version, Parsed->Flags);
+			if (Parsed->Version == 0)
+			{
+				ar.SkipBytes(2);			// `reserved`
+			}
+			else
+			{
+				ar.SkipBytes(1);			// `reserved`
+				uint8 Value8;
+				ar.Read(Value8);
+				Parsed->DefaultCryptByteBlock = Value8 >> 4;
+				Parsed->DefaultSkipByteBlock = Value8 & 0x0f;
+			}
+			ar.Read(Parsed->DefaultIsProtected);
+			ar.Read(Parsed->DefaultPerSampleIVSize);
+			Parsed->DefaultKID = MakeConstArrayView(ar.GetCurrentDataPointer(), 16);
+			ar.SkipBytes(16);
+			if (Parsed->DefaultIsProtected == 1 && Parsed->DefaultPerSampleIVSize == 0)
+			{
+				uint8 default_constant_IV_size;
+				ar.Read(default_constant_IV_size);
+				Parsed->DefaultConstantIV = MakeConstArrayView(ar.GetCurrentDataPointer(), default_constant_IV_size);
+				ar.SkipBytes(default_constant_IV_size);
+			}
+		}
+
+		bool FMP4BoxTENC::HasDefaultCryptBlockValues()
+		{
+			ParseIfRequired();
+			return Parsed->Version != 0;
+		}
+		uint8 FMP4BoxTENC::GetDefaultCryptByteBlock()
+		{
+			ParseIfRequired();
+			return Parsed->DefaultCryptByteBlock;
+		}
+		uint8 FMP4BoxTENC::GetDefaultSkipByteBlock()
+		{
+			ParseIfRequired();
+			return Parsed->DefaultSkipByteBlock;
+		}
+		uint8 FMP4BoxTENC::GetDefaultIsProtected()
+		{
+			ParseIfRequired();
+			return Parsed->DefaultIsProtected;
+		}
+		uint8 FMP4BoxTENC::GetDefaultPerSampleIVSize()
+		{
+			ParseIfRequired();
+			return Parsed->DefaultPerSampleIVSize;
+		}
+		TConstArrayView<uint8> FMP4BoxTENC::GetDefaultKID()
+		{
+			ParseIfRequired();
+			return Parsed->DefaultKID;
+		}
+		TConstArrayView<uint8> FMP4BoxTENC::GetDefaultConstantIV()
+		{
+			ParseIfRequired();
+			return Parsed->DefaultConstantIV;
+		}
+
+/****************************************************************************************************************************************************/
+
+		struct FMP4BoxSENC::FInitRefs
+		{
+			TArray<TSharedPtr<Electra::UtilitiesMP4::FMP4BoxBase>> RelatedBoxes;
+		};
+
+		struct FMP4BoxSENC::FParsed
+		{
+			TArray<FEntry> Entries;
+			uint32 Flags = 0;
+			uint8 Version = 0;
+		};
+		void FMP4BoxSENC::ParseIfRequired()
+		{
+			if (Parsed.IsValid() && InitRefs.IsValid())
+			{
+				return;
+			}
+			Parsed = MakePimpl<FParsed>();
+			FMP4AtomReaderHelper ar(BoxInfo.Data);
+			ar.ReadVersionAndFlags(Parsed->Version, Parsed->Flags);
+			if (InitRefs.IsValid())
+			{
+				// Check that we do not try to handle a pre-PIFF 1.3 box which has a different layout.
+				if ((Parsed->Flags & 1) != 0)
+				{
+					UE_LOG(LogElectraBase, Error, TEXT("The `senc` box is too old a version that is not supported"));
+					return;
+				}
+
+				/*
+					Depending on the version of this box and the scheme being used we will need access to a
+					number of other related boxes like `saiz`, `saio`, `seig` and possibly `ienc` and `iaux`.
+					At the very least for box versions 0 and 2 the related `tenc` box is required.
+
+					For the time being we only support version 0 where all samples are encrypted and a
+					`seig` box is not necessary.
+				*/
+				if (Parsed->Version != 0)
+				{
+					UE_LOG(LogElectraBase, Error, TEXT("At the moment only version 0 of the `senc` box is supported"));
+					return;
+				}
+
+				// We need the `Per_Sample_IV_Size` value from the `tenc` box.
+				TSharedPtr<Electra::UtilitiesMP4::FMP4BoxBase>* TencBoxPointer = InitRefs->RelatedBoxes.FindByPredicate([](const TSharedPtr<Electra::UtilitiesMP4::FMP4BoxBase>& InBox)
+				{
+					return InBox->GetType() == Electra::UtilitiesMP4::MakeBoxAtom('t','e','n','c');
+				});
+				if (!TencBoxPointer)
+				{
+					UE_LOG(LogElectraBase, Error, TEXT("Parsing a `senc` box requires the `tenc` box of the related track to be first passed via Prepare()"));
+					return;
+				}
+				TSharedPtr<Electra::UtilitiesMP4::FMP4BoxTENC> TencBox = StaticCastSharedPtr<Electra::UtilitiesMP4::FMP4BoxTENC>(*TencBoxPointer);
+				const uint8 Per_Sample_IV_Size = TencBox->GetDefaultPerSampleIVSize();
+
+				const bool senc_use_subsamples = !!(Parsed->Flags & 2);
+				const bool UseSubSampleEncryption = senc_use_subsamples /*&& Parsed->Version == 0*/;	// version already checked for above
+
+				uint32 sample_count;
+				ar.Read(sample_count);
+				Parsed->Entries.SetNum(sample_count);
+				for(uint32 i=0; i<sample_count; ++i)
+				{
+					FEntry& e = Parsed->Entries[i];
+					/*
+					if (Parsed->Version == 0)
+					{
+					*/
+						if (Per_Sample_IV_Size)
+						{
+							e.IV = MakeConstArrayView(ar.GetCurrentDataPointer(), Per_Sample_IV_Size);
+							ar.SkipBytes(Per_Sample_IV_Size);
+						}
+						else
+						{
+							// For convenience we set the contant IV from the `tenc` box here.
+							e.IV = TencBox->GetDefaultConstantIV();
+						}
+						if (UseSubSampleEncryption)
+						{
+							uint16 subsample_count;
+							ar.Read(subsample_count);
+							e.SubSamples.SetNum(subsample_count);
+							for(int32 k=0; k<subsample_count; ++k)
+							{
+								ar.Read(e.SubSamples[k].NumClearBytes);
+								ar.Read(e.SubSamples[k].NumEncryptedBytes);
+							}
+						}
+					/*
+					}
+					else if (Parsed->Version == 1 && isProtected)
+					{
+					}
+					else if (Parsed->Version == 2 && isProtected)
+					{
+					}
+					*/
+				}
+			}
+			else
+			{
+				UE_LOG(LogElectraBase, Error, TEXT("Accessing a `senc` box that has not been prepared with the corresponding `tenc` and other related boxes"));
+			}
+		}
+
+		void FMP4BoxSENC::Prepare(const TArray<TSharedPtr<Electra::UtilitiesMP4::FMP4BoxBase>>& InRelatedBoxes)
+		{
+			// If already prepared we're done.
+			if (!InitRefs.IsValid())
+			{
+				InitRefs = MakePimpl<FInitRefs>();
+				InitRefs->RelatedBoxes = InRelatedBoxes;
+			}
+		}
+
+		const TArray<FMP4BoxSENC::FEntry>& FMP4BoxSENC::GetEntries()
+		{
+			ParseIfRequired();
+			return Parsed->Entries;
+		}
+
+/****************************************************************************************************************************************************/
+
+		struct FMP4BoxSIDX::FParsed
+		{
+			FEntryList Entries;
+			uint64 EarliestPresentationTime = 0;
+			uint64 FirstOffset = 0;
+			uint32 ReferenceID = 0;
+			uint32 Timescale = 0;
+			uint32 Flags = 0;
+			uint8 Version = 0;
+		};
+		void FMP4BoxSIDX::ParseIfRequired()
+		{
+			if (Parsed.IsValid())
+			{
+				return;
+			}
+			Parsed = MakePimpl<FParsed>();
+			FMP4AtomReaderHelper ar(BoxInfo.Data);
+			ar.ReadVersionAndFlags(Parsed->Version, Parsed->Flags);
+			ar.Read(Parsed->ReferenceID);
+			ar.Read(Parsed->Timescale);
+			if (Parsed->Version == 0)
+			{
+				uint32 Value32;
+				ar.Read(Value32);
+				Parsed->EarliestPresentationTime = Value32;
+				ar.Read(Value32);
+				Parsed->FirstOffset = Value32;
+			}
+			else
+			{
+				ar.Read(Parsed->EarliestPresentationTime);
+				ar.Read(Parsed->FirstOffset);
+			}
+			ar.SkipBytes(2);		// `reserved`
+			uint16 reference_count;
+			ar.Read(reference_count);
+			Parsed->Entries.Reserve(reference_count);
+			static_assert(sizeof(FEntry::SubSegmentDuration) == 4);
+			for(int32 i=0; i<reference_count; ++i)
+			{
+				FEntry e;
+				uint32 ReferenceTypeAndSize;
+				uint32 SAPStartAndTypeAndDeltaTime;
+				ar.Read(ReferenceTypeAndSize);
+				ar.Read(e.SubSegmentDuration);
+				ar.Read(SAPStartAndTypeAndDeltaTime);
+				e.IsReferenceType = ReferenceTypeAndSize >> 31;
+				e.Size = ReferenceTypeAndSize & 0x7fffffff;
+				e.StartsWithSAP  = SAPStartAndTypeAndDeltaTime >> 31;
+				e.SAPType = (SAPStartAndTypeAndDeltaTime >> 28) & 7;
+				e.SAPDeltaTime = SAPStartAndTypeAndDeltaTime & 0x0fffffffU;
+				Parsed->Entries.AddElement(MoveTemp(e));
+			}
+		}
+
+		uint32 FMP4BoxSIDX::GetReferenceID()
+		{
+			ParseIfRequired();
+			return Parsed->ReferenceID;
+		}
+		uint32 FMP4BoxSIDX::GetTimescale()
+		{
+			ParseIfRequired();
+			return Parsed->Timescale;
+		}
+		uint64 FMP4BoxSIDX::GetEarliestPresentationTime()
+		{
+			ParseIfRequired();
+			return Parsed->EarliestPresentationTime;
+		}
+		uint64 FMP4BoxSIDX::GetFirstOffset()
+		{
+			ParseIfRequired();
+			return Parsed->FirstOffset;
+		}
+		const FMP4BoxSIDX::FEntryList& FMP4BoxSIDX::GetEntries()
+		{
+			ParseIfRequired();
+			return Parsed->Entries;
+		}
+
+
+/****************************************************************************************************************************************************/
+#if 0 // TEMPLATE
+		struct FMP4BoxCCCC::FParsed
+		{
+			uint32 Flags = 0;
+			uint8 Version = 0;
+		};
+		void FMP4BoxCCCC::ParseIfRequired()
+		{
+			if (Parsed.IsValid())
+			{
+				return;
+			}
+			Parsed = MakePimpl<FParsed>();
+			FMP4AtomReaderHelper ar(BoxInfo.Data);
+			ar.ReadVersionAndFlags(Parsed->Version, Parsed->Flags);
+			if (Parsed->Version == 1)
+			{
+			}
+			else
+			{
+			}
+		}
+
+		bool FMP4BoxCCCC::XXXX()
+		{
+			ParseIfRequired();
+			return
+		}
+#endif
+/****************************************************************************************************************************************************/
+
+		FMP4BoxFactory& FMP4BoxFactory::Get()
+		{
+			static FMP4BoxFactory Factory;
+			return Factory;
+		}
+
+		FMP4BoxFactory::FMP4BoxFactory()
+		{
+			FactoryMap.Add(MakeBoxAtom('f','t','y','p'), &FMP4BoxFTYP::Create);
+			FactoryMap.Add(MakeBoxAtom('s','t','y','p'), &FMP4BoxFTYP::Create);
+			FactoryMap.Add(MakeBoxAtom('m','o','o','v'), &FMP4BoxMOOV::Create);
+			FactoryMap.Add(MakeBoxAtom('m','v','h','d'), &FMP4BoxMVHD::Create);
+			FactoryMap.Add(MakeBoxAtom('t','r','a','k'), &FMP4BoxTRAK::Create);
+			FactoryMap.Add(MakeBoxAtom('t','k','h','d'), &FMP4BoxTKHD::Create);
+			FactoryMap.Add(MakeBoxAtom('t','r','e','f'), &FMP4BoxTREF::Create);
+			FactoryMap.Add(MakeBoxAtom('e','d','t','s'), &FMP4BoxEDTS::Create);
+			FactoryMap.Add(MakeBoxAtom('e','l','s','t'), &FMP4BoxELST::Create);
+			FactoryMap.Add(MakeBoxAtom('m','d','i','a'), &FMP4BoxMDIA::Create);
+			FactoryMap.Add(MakeBoxAtom('m','d','h','d'), &FMP4BoxMDHD::Create);
+			FactoryMap.Add(MakeBoxAtom('h','d','l','r'), &FMP4BoxHDLR::Create);
+			FactoryMap.Add(MakeBoxAtom('e','l','n','g'), &FMP4BoxHDLR::Create);
+			FactoryMap.Add(MakeBoxAtom('m','i','n','f'), &FMP4BoxMINF::Create);
+			FactoryMap.Add(MakeBoxAtom('d','i','n','f'), &FMP4BoxDINF::Create);
+			FactoryMap.Add(MakeBoxAtom('d','r','e','f'), &FMP4BoxDREF::Create);
+			FactoryMap.Add(MakeBoxAtom('s','t','b','l'), &FMP4BoxSTBL::Create);
+			FactoryMap.Add(MakeBoxAtom('s','t','s','d'), &FMP4BoxSTSD::Create);
+			FactoryMap.Add(MakeBoxAtom('s','t','t','s'), &FMP4BoxSTTS::Create);
+			FactoryMap.Add(MakeBoxAtom('c','t','t','s'), &FMP4BoxCTTS::Create);
+			FactoryMap.Add(MakeBoxAtom('c','s','l','g'), &FMP4BoxCSLG::Create);
+			FactoryMap.Add(MakeBoxAtom('s','t','s','s'), &FMP4BoxSTSS::Create);
+			FactoryMap.Add(MakeBoxAtom('s','d','t','p'), &FMP4BoxSDTP::Create);
+			FactoryMap.Add(MakeBoxAtom('s','t','s','c'), &FMP4BoxSTSC::Create);
+			FactoryMap.Add(MakeBoxAtom('s','t','s','z'), &FMP4BoxSTSZ::Create);
+			FactoryMap.Add(MakeBoxAtom('s','t','c','o'), &FMP4BoxSTCO::Create);
+			FactoryMap.Add(MakeBoxAtom('c','o','6','4'), &FMP4BoxSTCO::Create);
+			FactoryMap.Add(MakeBoxAtom('s','a','i','z'), &FMP4BoxSAIZ::Create);
+			FactoryMap.Add(MakeBoxAtom('s','a','i','o'), &FMP4BoxSAIO::Create);
+			FactoryMap.Add(MakeBoxAtom('s','g','p','d'), &FMP4BoxSGPD::Create);
+			FactoryMap.Add(MakeBoxAtom('s','b','g','p'), &FMP4BoxSBGP::Create);
+
+			FactoryMap.Add(MakeBoxAtom('m','v','e','x'), &FMP4BoxMVEX::Create);
+			FactoryMap.Add(MakeBoxAtom('m','e','h','d'), &FMP4BoxMEHD::Create);
+			FactoryMap.Add(MakeBoxAtom('t','r','e','x'), &FMP4BoxTREX::Create);
+			FactoryMap.Add(MakeBoxAtom('l','e','v','a'), &FMP4BoxLEVA::Create);
+			FactoryMap.Add(MakeBoxAtom('m','o','o','f'), &FMP4BoxMOOF::Create);
+			FactoryMap.Add(MakeBoxAtom('m','f','h','d'), &FMP4BoxMFHD::Create);
+			FactoryMap.Add(MakeBoxAtom('t','r','a','f'), &FMP4BoxTRAF::Create);
+			FactoryMap.Add(MakeBoxAtom('t','f','h','d'), &FMP4BoxTFHD::Create);
+			FactoryMap.Add(MakeBoxAtom('t','r','u','n'), &FMP4BoxTRUN::Create);
+			FactoryMap.Add(MakeBoxAtom('m','f','r','a'), &FMP4BoxMFRA::Create);
+			FactoryMap.Add(MakeBoxAtom('t','f','r','a'), &FMP4BoxTFRA::Create);
+			FactoryMap.Add(MakeBoxAtom('m','f','r','o'), &FMP4BoxMFRO::Create);
+			FactoryMap.Add(MakeBoxAtom('t','f','d','t'), &FMP4BoxTFDT::Create);
+			FactoryMap.Add(MakeBoxAtom('s','i','d','x'), &FMP4BoxSIDX::Create);
+			FactoryMap.Add(MakeBoxAtom('s','s','i','x'), &FMP4BoxSSIX::Create);
+
+			FactoryMap.Add(MakeBoxAtom('f','r','e','e'), &FMP4BoxFREE::Create);
+			FactoryMap.Add(MakeBoxAtom('s','k','i','p'), &FMP4BoxFREE::Create);
+			FactoryMap.Add(MakeBoxAtom('u','d','t','a'), &FMP4BoxUDTA::Create);
+			FactoryMap.Add(MakeBoxAtom('m','e','t','a'), &FMP4BoxMETA::Create);
+
+			FactoryMap.Add(MakeBoxAtom('v','m','h','d'), &FMP4BoxVMHD::Create);
+			FactoryMap.Add(MakeBoxAtom('s','m','h','d'), &FMP4BoxSMHD::Create);
+			FactoryMap.Add(MakeBoxAtom('n','m','h','d'), &FMP4BoxNMHD::Create);
+
+			FactoryMap.Add(MakeBoxAtom('g','m','h','d'), &FMP4BoxGMHD::Create);
+			FactoryMap.Add(MakeBoxAtom('t','a','p','t'), &FMP4BoxTAPT::Create);
+
+			FactoryMap.Add(MakeBoxAtom('c','o','l','r'), &FMP4BoxCOLR::Create);
+			FactoryMap.Add(MakeBoxAtom('c','l','l','i'), &FMP4BoxCLLI::Create);
+			FactoryMap.Add(MakeBoxAtom('m','d','c','v'), &FMP4BoxMDCV::Create);
+			FactoryMap.Add(MakeBoxAtom('b','t','r','t'), &FMP4BoxBTRT::Create);
+			FactoryMap.Add(MakeBoxAtom('p','a','s','p'), &FMP4BoxPASP::Create);
+
+			FactoryMap.Add(MakeBoxAtom('s','i','n','f'), &FMP4BoxSINF::Create);
+			FactoryMap.Add(MakeBoxAtom('f','r','m','a'), &FMP4BoxFRMA::Create);
+			FactoryMap.Add(MakeBoxAtom('s','c','h','m'), &FMP4BoxSCHM::Create);
+			FactoryMap.Add(MakeBoxAtom('s','c','h','i'), &FMP4BoxSCHI::Create);
+
+			FactoryMap.Add(MakeBoxAtom('p','s','s','h'), &FMP4BoxPSSH::Create);
+			FactoryMap.Add(MakeBoxAtom('t','e','n','c'), &FMP4BoxTENC::Create);
+			FactoryMap.Add(MakeBoxAtom('s','e','n','c'), &FMP4BoxSENC::Create);
+
+			// `mdat` box is never parsed, we only handle it as a base box here to create the box tree.
+			FactoryMap.Add(MakeBoxAtom('m','d','a','t'), &FMP4BoxBase::Create);
+
+//			FactoryMap.Add(MakeBoxAtom('u','r','l',' '), &FMP4BoxBase::Create);
+
+
+			// Derived formats
+			FactoryMap.Add(MakeBoxAtom('a','v','c','C'), &FMP4BoxAVCC::Create);
+			FactoryMap.Add(MakeBoxAtom('h','v','c','C'), &FMP4BoxHVCC::Create);
+			FactoryMap.Add(MakeBoxAtom('d','v','c','C'), &FMP4BoxDVCC::Create);
+			FactoryMap.Add(MakeBoxAtom('d','v','v','C'), &FMP4BoxDVCC::Create);
+			FactoryMap.Add(MakeBoxAtom('d','v','w','C'), &FMP4BoxDVCC::Create);
+			FactoryMap.Add(MakeBoxAtom('d','a','c','3'), &FMP4BoxDAC3::Create);
+			FactoryMap.Add(MakeBoxAtom('d','e','c','3'), &FMP4BoxDEC3::Create);
+			FactoryMap.Add(MakeBoxAtom('i','o','d','s'), &FMP4BoxIODS::Create);
+			FactoryMap.Add(MakeBoxAtom('e','s','d','s'), &FMP4BoxESDS::Create);
+			FactoryMap.Add(MakeBoxAtom('w','a','v','e'), &FMP4BoxWAVE::Create);
+			FactoryMap.Add(MakeBoxAtom('d','f','L','a'), &FMP4BoxDFLA::Create);
+			FactoryMap.Add(MakeBoxAtom('d','O','p','s'), &FMP4BoxDOPS::Create);
+			FactoryMap.Add(MakeBoxAtom('v','p','c','C'), &FMP4BoxVPCC::Create);
+		}
+		TSharedPtr<FMP4BoxBase> FMP4BoxFactory::Create(const TWeakPtr<FMP4BoxBase>& InParent, const FMP4BoxInfo& InBoxInfo)
+		{
+			if (auto Factory = FactoryMap.Find(InBoxInfo.Type))
+			{
+				return (*Factory)(InParent, InBoxInfo);
+			}
+
+			// Check for an `uuid` box. We know some of them and can handle them appropriately.
+			if (InBoxInfo.Type == MakeBoxAtom('u','u','i','d'))
+			{
+				static const uint8 UUID_atom[16] { 0x00,0x00,0x00,0x00,0x00,0x11,0x00,0x10,0x80,0x00,0x00,0xAA,0x00,0x38,0x9B,0x71 };
+				static const uint8 UUID_pssh[16] { 0xD0,0x8A,0x4F,0x18,0x10,0xF3,0x4A,0x82,0xB6,0xC8,0x32,0xD8,0xAB,0xA1,0x83,0xD3 };
+				static const uint8 UUID_tenc[16] { 0x89,0x74,0xDB,0xCE,0x7B,0xE7,0x4C,0x51,0x84,0xF9,0x71,0x48,0xF9,0x88,0x25,0x54 };
+				static const uint8 UUID_senc[16] { 0xA2,0x39,0x4F,0x52,0x5A,0x9B,0x4F,0x14,0xA2,0x44,0x6C,0x42,0x7C,0x64,0x8D,0xF4 };
+				if (FMemory::Memcmp(InBoxInfo.UUID+4, UUID_atom+4, 12) == 0)
+				{
+					// This is not handled.
+					UE_LOG(LogElectraBase, Error, TEXT("A `uuid` box that uses the long form of a well known atom is not supported"));
+				}
+				// Compare the UUID with known boxes.
+				else if (FMemory::Memcmp(InBoxInfo.UUID, UUID_pssh, 16) == 0)
+				{
+					FMP4BoxInfo WellKnown(InBoxInfo);
+					WellKnown.Type = MakeBoxAtom('p','s','s','h');
+					return FMP4BoxPSSH::Create(InParent, WellKnown);
+				}
+				else if (FMemory::Memcmp(InBoxInfo.UUID, UUID_tenc, 16) == 0)
+				{
+					FMP4BoxInfo WellKnown(InBoxInfo);
+					WellKnown.Type = MakeBoxAtom('t','e','n','c');
+					return FMP4BoxTENC::Create(InParent, WellKnown);
+				}
+				else if (FMemory::Memcmp(InBoxInfo.UUID, UUID_senc, 16) == 0)
+				{
+					FMP4BoxInfo WellKnown(InBoxInfo);
+					WellKnown.Type = MakeBoxAtom('s','e','n','c');
+					return FMP4BoxSENC::Create(InParent, WellKnown);
+				}
+				else
+				{
+					// Unknown UUID, generate a base box.
+					return FMP4BoxBase::Create(InParent, InBoxInfo);
+				}
+			}
+
+			//UE_LOG(LogTemp, Log, TEXT("unhandled box \"%s\""), *FString::ConstructFromPtrSize((const ANSICHAR*)InBoxInfo.Name, 4));
+
+			return FMP4BoxBase::Create(InParent, InBoxInfo);
+		}
+
+
+	} // namespace UtilitiesMP4
+
+} // namespace Electra
diff --git a/Engine/Plugins/Media/ElectraUtil/Source/ElectraBase/Private/Utilities/MP4Boxes/MP4Track.cpp b/Engine/Plugins/Media/ElectraUtil/Source/ElectraBase/Private/Utilities/MP4Boxes/MP4Track.cpp
new file mode 100644
--- /dev/null
+++ b/Engine/Plugins/Media/ElectraUtil/Source/ElectraBase/Private/Utilities/MP4Boxes/MP4Track.cpp
@@ -0,0 +1,602 @@
+// Copyright Epic Games, Inc. All Rights Reserved.
+
+#include "Utilities/MP4Boxes/MP4Track.h"
+#include "ElectraBaseModule.h"	// for log category `LogElectraBase`
+
+namespace Electra
+{
+	namespace UtilitiesMP4
+	{
+
+		bool FMP4Track::Prepare(FTimeFraction InMovieDurationFromMVHDBox)
+		{
+			if (!TrakBox.IsValid())
+			{
+				LastErrorMessage = FString::Printf(TEXT("No `trak` box given."));
+				return false;
+			}
+			TkhdBox = TrakBox->FindBoxRecursive<Electra::UtilitiesMP4::FMP4BoxTKHD>(Electra::UtilitiesMP4::MakeBoxAtom('t','k','h','d'), 0);
+			if (!TkhdBox.IsValid())
+			{
+				LastErrorMessage = FString::Printf(TEXT("No `tkhd` box in `trak`."));
+				return false;
+			}
+			Convs.TrackID = TkhdBox->GetTrackID();
+			Convs.DurationFromMvhdBox = InMovieDurationFromMVHDBox;
+			Convs.DurationFromTkhdBox.SetFromND(TkhdBox->GetDuration(), Convs.DurationFromMvhdBox.GetDenominator());
+
+			// Check for correct box hierarchy.
+			TSharedPtr<FMP4BoxMDIA> MdiaBox = TrakBox->FindBoxRecursive<Electra::UtilitiesMP4::FMP4BoxMDIA>(Electra::UtilitiesMP4::MakeBoxAtom('m','d','i','a'), 0);
+			if (!MdiaBox.IsValid())
+			{
+				LastErrorMessage = FString::Printf(TEXT("No `mdia` box in `trak`."));
+				return false;
+			}
+			TSharedPtr<FMP4BoxMINF> MinfBox = MdiaBox->FindBoxRecursive<Electra::UtilitiesMP4::FMP4BoxMINF>(Electra::UtilitiesMP4::MakeBoxAtom('m','i','n','f'), 0);
+			if (!MinfBox.IsValid())
+			{
+				LastErrorMessage = FString::Printf(TEXT("No `minf` box in `mdia`."));
+				return false;
+			}
+			TSharedPtr<FMP4BoxSTBL> StblBox = MinfBox->FindBoxRecursive<Electra::UtilitiesMP4::FMP4BoxSTBL>(Electra::UtilitiesMP4::MakeBoxAtom('s','t','b','l'), 0);
+			if (!StblBox.IsValid())
+			{
+				LastErrorMessage = FString::Printf(TEXT("No `stbl` box in `minf`."));
+				return false;
+			}
+
+			MdhdBox = MdiaBox->FindBoxRecursive<Electra::UtilitiesMP4::FMP4BoxMDHD>(Electra::UtilitiesMP4::MakeBoxAtom('m','d','h','d'), 0);
+			if (!MdhdBox.IsValid())
+			{
+				LastErrorMessage = FString::Printf(TEXT("No `mdhd` box in `mdia`."));
+				return false;
+			}
+			Convs.DurationFromMdhdBox = MdhdBox->GetDuration();
+			if (Convs.DurationFromMdhdBox.GetDenominator() == 0)
+			{
+				LastErrorMessage = FString::Printf(TEXT("Timescale in `mdhd` box is zero, which is not supported."));
+				return false;
+			}
+
+			// Required sample information boxes:
+			SttsBox = StblBox->FindBoxRecursive<Electra::UtilitiesMP4::FMP4BoxSTTS>(Electra::UtilitiesMP4::MakeBoxAtom('s','t','t','s'), 0);
+			if (!SttsBox.IsValid())
+			{
+				LastErrorMessage = FString::Printf(TEXT("No `stts` box in `stbl`."));
+				return false;
+			}
+			StscBox = StblBox->FindBoxRecursive<Electra::UtilitiesMP4::FMP4BoxSTSC>(Electra::UtilitiesMP4::MakeBoxAtom('s','t','s','c'), 0);
+			if (!StscBox.IsValid())
+			{
+				LastErrorMessage = FString::Printf(TEXT("No `stsc` box in `stbl`."));
+				return false;
+			}
+			StszBox = StblBox->FindBoxRecursive<Electra::UtilitiesMP4::FMP4BoxSTSZ>(Electra::UtilitiesMP4::MakeBoxAtom('s','t','s','z'), 0);
+			if (!StszBox.IsValid())
+			{
+				LastErrorMessage = FString::Printf(TEXT("No `stsz` box in `stbl`."));
+				return false;
+			}
+			StcoBox = StblBox->FindBoxRecursive<Electra::UtilitiesMP4::FMP4BoxSTCO>(Electra::UtilitiesMP4::MakeBoxAtom('s','t','c','o'), 0);
+			if (!StcoBox.IsValid())
+			{
+				LastErrorMessage = FString::Printf(TEXT("No `stco` ot `co64` box in `stbl`."));
+				return false;
+			}
+
+			// Validity check
+			Convs.NumTotalSamples = StszBox->GetNumberOfSamples();
+			if (Convs.NumTotalSamples != SttsBox->GetNumTotalSamples())
+			{
+				LastErrorMessage = FString::Printf(TEXT("Mismatching number of samples in `stts` and `stsz` boxes."));
+				return false;
+			}
+
+			// Optional sample information boxes:
+			CttsBox = StblBox->FindBoxRecursive<Electra::UtilitiesMP4::FMP4BoxCTTS>(Electra::UtilitiesMP4::MakeBoxAtom('c','t','t','s'), 0);
+			StssBox = StblBox->FindBoxRecursive<Electra::UtilitiesMP4::FMP4BoxSTSS>(Electra::UtilitiesMP4::MakeBoxAtom('s','t','s','s'), 0);
+			StblBox->GetAllBoxInstances(SgpdBoxes, Electra::UtilitiesMP4::MakeBoxAtom('s','g','p','d'));
+			StblBox->GetAllBoxInstances(SbgpBoxes, Electra::UtilitiesMP4::MakeBoxAtom('s','b','g','p'));
+
+			// Start with default values for what is mapped onto the timeline.
+			Convs.CompositionTimeAtZeroPoint = CttsBox.IsValid() && CttsBox->GetEntries().Num() ? CttsBox->GetEntries()[0].sample_offset : 0;
+			Convs.MappedDurationFromElstBox = Convs.DurationFromMvhdBox;
+
+			// Optional edit list
+			ElstBox = TrakBox->FindBoxRecursive<Electra::UtilitiesMP4::FMP4BoxELST>(Electra::UtilitiesMP4::MakeBoxAtom('e','l','s','t'), 1);
+			// If there is an edit list it needs to be simple and only contain a composition time mapping.
+			if (ElstBox.IsValid())
+			{
+				if (ElstBox->RepeatEdits())
+				{
+					LastErrorMessage = FString::Printf(TEXT("Repeating `elst` box ist not supported."));
+					return false;
+				}
+				// If there is more than a single entry things are bound to be complicated, so we don't even
+				// want to know if the entry we are interested in is there.
+				const TArray<Electra::UtilitiesMP4::FMP4BoxELST::FEntry>& ElstEntries = ElstBox->GetEntries();
+				if (ElstEntries.Num() == 0)
+				{
+					LastErrorMessage = FString::Printf(TEXT("Edit list is empty."));
+					return false;
+				}
+				if (ElstEntries.Num() > 1)
+				{
+					LastErrorMessage = FString::Printf(TEXT("Edit list with more than one entry is not supported."));
+					return false;
+				}
+				if (!(ElstEntries[0].media_rate_integer == 1 && ElstEntries[0].media_rate_fraction == 0))
+				{
+					LastErrorMessage = FString::Printf(TEXT("Edit list entries with playback rates other than 1.0 are not supported."));
+					return false;
+				}
+				if (ElstEntries[0].media_time < 0)
+				{
+					LastErrorMessage = FString::Printf(TEXT("Edit list specifies an empty edit, which is not supported."));
+					return false;
+				}
+
+				Convs.CompositionTimeAtZeroPoint = ElstEntries[0].media_time;
+				// The mapped duration may be different from the media duration itself, in which case the track is either truncated
+				// or the last sample has a longer duration to be repeated until the end.
+				if (ElstEntries[0].edit_duration == 0)
+				{
+					// The value of 0 is reserved for fragmented files with no `mehd` box.
+					// We do not handle this case.
+					LastErrorMessage = FString::Printf(TEXT("Edit list specifies zero edit duration, which is not supported."));
+					return false;
+				}
+				Convs.MappedDurationFromElstBox.SetFromND(ElstEntries[0].edit_duration, Convs.DurationFromMvhdBox.GetDenominator());
+			}
+			// For convenience sake convert the mapped duration from `mvhd` timescale into `mdhd` timescale.
+			Convs.MappedDurationFromElstBox.SetFromND(Convs.MappedDurationFromElstBox.GetAsTimebase(Convs.DurationFromMdhdBox.GetDenominator()), Convs.DurationFromMdhdBox.GetDenominator());
+
+
+			// Set up common track meta data
+			UdtaBox = TrakBox->FindBoxRecursive<Electra::UtilitiesMP4::FMP4BoxUDTA>(Electra::UtilitiesMP4::MakeBoxAtom('u','d','t','a'), 0);
+			if (UdtaBox.IsValid())
+			{
+				// Is there a `name` box?
+				TSharedPtr<FMP4BoxBase> NameBox = UdtaBox->FindBoxRecursive<Electra::UtilitiesMP4::FMP4BoxBase>(Electra::UtilitiesMP4::MakeBoxAtom('n','a','m','e'), 0);
+				if (NameBox.IsValid() && NameBox->GetBoxData().Num())
+				{
+					CommonMetadata.Name = FString::ConstructFromPtrSize(reinterpret_cast<const ANSICHAR*>(NameBox->GetBoxData().GetData()), NameBox->GetBoxData().Num());
+				}
+			}
+			TSharedPtr<FMP4BoxHDLR> HdlrBox = MdiaBox->FindBoxRecursive<Electra::UtilitiesMP4::FMP4BoxHDLR>(Electra::UtilitiesMP4::MakeBoxAtom('h','d','l','r'), 0);
+			check(HdlrBox.IsValid())
+			if (HdlrBox.IsValid())
+			{
+				CommonMetadata.HandlerName = HdlrBox->GetHandlerName();
+			}
+			// Get language from `mdhd` box first.
+			CommonMetadata.Language = MdhdBox->GetLanguage();
+			// Then, if there is an `elng` box that gives better information get it from there.
+			TSharedPtr<FMP4BoxELNG> ElngBox = MdiaBox->FindBoxRecursive<Electra::UtilitiesMP4::FMP4BoxELNG>(Electra::UtilitiesMP4::MakeBoxAtom('e','l','n','g'), 0);
+			if (ElngBox.IsValid())
+			{
+				CommonMetadata.Language = ElngBox->GetLanguage();
+			}
+
+			bHasBeenPrepared = true;
+
+			// Given the timeline mapping, locate the sample number that falls onto the start of the timeline
+			// and last one falling onto the end of the timeline.
+			TSharedPtr<FIterator> StartIt = CreateIterator();
+			check(StartIt.IsValid());
+			const int64 pts0 = Convs.CompositionTimeAtZeroPoint;
+			uint32 SyncSampleNum = 0;
+			bool bFound = false;
+			while(1)
+			{
+				if (StartIt->IsSyncOrRAPSample())
+				{
+					SyncSampleNum = StartIt->GetSampleNumber();
+				}
+				int64 s = StartIt->GetPTS().GetNumerator();
+				int64 e = s + StartIt->GetDuration().GetNumerator();
+				if (pts0 >= s && pts0 < e)
+				{
+					Convs.FirstSample.SampleNumber = StartIt->GetSampleNumber();
+					Convs.FirstSample.SamplePTS = s;
+					Convs.FirstSample.StartPTS = pts0;
+					Convs.FirstSample.SyncSampleNumber = SyncSampleNum;
+					Convs.DTSShiftAtZeroPoint = StartIt->GetDTS().GetNumerator();
+					bFound = true;
+					break;
+				}
+				if (!StartIt->Next())
+				{
+					break;
+				}
+			}
+			check(bFound);
+
+			TSharedPtr<FIterator> EndIt = CreateIteratorAtLastFrame();
+			check(EndIt.IsValid());
+
+			// Find the highest PTS
+			uint32 HighestPTSIndex = ~0U;
+			int64 HighestPTS = TNumericLimits<int64>::Min();
+			int64 HighestEndPTS = 0;
+			while(1)
+			{
+				int64 pts = EndIt->GetPTS().GetNumerator();
+				if (pts > HighestPTS)
+				{
+					HighestPTS = pts;
+					HighestEndPTS = pts + EndIt->GetDuration().GetNumerator();
+					HighestPTSIndex = EndIt->GetSampleNumber();
+				}
+				if (EndIt->IsSyncOrRAPSample() || !EndIt->Prev())
+				{
+					break;
+				}
+			}
+			check(HighestPTSIndex != ~0U);
+
+			const int64 pts1 = Convs.CompositionTimeAtZeroPoint + Convs.MappedDurationFromElstBox.GetNumerator();
+			// Is the mapped duration is greater or equal to what the media duration is?
+			if (pts1 >= HighestEndPTS)
+			{
+				Convs.LastSample.SampleNumber = HighestPTSIndex;
+				Convs.LastSample.LastSampleNumber = Convs.NumTotalSamples - 1;
+				Convs.LastSample.SamplePTS = HighestPTS;
+				Convs.LastSample.EndPTS = pts1;
+				double PaddingDuration = (pts1 - HighestEndPTS) / (double)Convs.DurationFromMdhdBox.GetDenominator();
+				if (PaddingDuration >= 0.001)
+				{
+					UE_LOG(LogElectraBase, Verbose, TEXT("Last sample duration in track #%u will be extended by %#.5f seconds to align with the movie duration in the `mvhd` box."), Convs.TrackID, PaddingDuration);
+				}
+			}
+			// The mapping truncates the media. Find where that is.
+			else
+			{
+				EndIt = CreateIteratorAtLastFrame();
+				bFound = false;
+				while(!bFound)
+				{
+					int64 s = EndIt->GetPTS().GetNumerator();
+					int64 e = s + EndIt->GetDuration().GetNumerator();
+					if (pts1 > s && pts1 <= e)
+					{
+						// This sample contains the end of the mapped duration.
+						Convs.LastSample.SamplePTS = s;
+						Convs.LastSample.EndPTS = pts1;
+						Convs.LastSample.SampleNumber = EndIt->GetSampleNumber();
+						// If it is a sync or rap sample then this is also the last sample we need to look at
+						// as no later sample (in decode order) may be needed to decode this one.
+						if (EndIt->IsSyncOrRAPSample() || EndIt->IsLast())
+						{
+							Convs.LastSample.LastSampleNumber = EndIt->GetSampleNumber();
+						}
+						else
+						{
+							// We need decode frames up to this PTS, meaning that everything that comes
+							// earlier in decode order we need to decode as well.
+							Convs.LastSample.LastSampleNumber = EndIt->GetSampleNumber();
+							while(EndIt->Next())
+							{
+								int64 NextS = EndIt->GetDTS().GetNumerator();
+								if (NextS > pts1)
+								{
+									break;
+								}
+								Convs.LastSample.LastSampleNumber = EndIt->GetSampleNumber();
+							}
+						}
+						bFound = true;
+					}
+					else if (!EndIt->Prev())
+					{
+						break;
+					}
+				}
+				check(bFound);
+			}
+			return true;
+		}
+
+		const FMP4TrackMetadataCommon& FMP4Track::GetCommonMetadata()
+		{
+			return CommonMetadata;
+		}
+
+
+		FString FMP4Track::GetLastError()
+		{
+			return LastErrorMessage;
+		}
+
+
+		// Returns the number of samples in this track.
+		uint32 FMP4Track::GetNumberOfSamples()
+		{
+			check(bHasBeenPrepared);
+			return bHasBeenPrepared ? Convs.NumTotalSamples : 0;
+		}
+
+		// Create an interator starting at the first sample.
+		TSharedPtr<FMP4Track::FIterator> FMP4Track::CreateIterator()
+		{
+			check(bHasBeenPrepared);
+			if (!bHasBeenPrepared)
+			{
+				LastErrorMessage = FString::Printf(TEXT("Track has not been prepared, cannot create an iterator."));
+				return nullptr;
+			}
+			if (Convs.NumTotalSamples == 0)
+			{
+				LastErrorMessage = FString::Printf(TEXT("There are no samples in this track, cannot create an iterator."));
+				return nullptr;
+			}
+			TUniquePtr<FIterator> It(new FIterator);
+			It->Track = AsWeak();
+			It->SampleNumber = 0;
+			// Copy convenience values into the iterator
+			It->Convs = Convs;
+			// Set up the box iterators.
+			It->stszIt.SetBox(StszBox);
+			It->sttsIt.SetBox(SttsBox);
+			It->cttsIt.SetBox(CttsBox, Convs.NumTotalSamples);
+			It->stcoIt.SetBox(StcoBox);
+			It->stscIt.SetBox(StscBox, Convs.NumTotalSamples);
+			It->stssIt.SetBox(StssBox, Convs.NumTotalSamples);
+			// Do we have a `rap ` group?
+			TSharedPtr<FMP4BoxSGPD>* rapSGPD = SgpdBoxes.FindByPredicate([](const TSharedPtr<FMP4BoxSGPD>& e){ return e->GetGroupingType() == Electra::UtilitiesMP4::MakeBoxAtom('r','a','p',' ');});
+			TSharedPtr<FMP4BoxSBGP>* rapSBGP = SbgpBoxes.FindByPredicate([](const TSharedPtr<FMP4BoxSBGP>& e){ return e->GetGroupingType() == Electra::UtilitiesMP4::MakeBoxAtom('r','a','p',' ');});
+			if (rapSGPD && rapSBGP)
+			{
+				It->rapIt.SetBox(*rapSBGP, (*rapSGPD)->GetDefaultGroupDescriptionIndex(), Convs.NumTotalSamples);
+			}
+			else
+			{
+				// Initialize the iterator such that it can be used to return "not a RAP" for every sample.
+				It->rapIt.SetBox(nullptr, 0, Convs.NumTotalSamples);
+			}
+			It->Update();
+			return MakeShareable(It.Release());
+		}
+
+		// Create an interator starting at the last sample (used when iterating in reverse, crossing back from the beginning to the end)
+		TSharedPtr<FMP4Track::FIterator> FMP4Track::CreateIteratorAtLastFrame()
+		{
+			check(bHasBeenPrepared);
+			if (!bHasBeenPrepared)
+			{
+				LastErrorMessage = FString::Printf(TEXT("Track has not been prepared, cannot create an iterator."));
+				return nullptr;
+			}
+			if (Convs.NumTotalSamples == 0)
+			{
+				LastErrorMessage = FString::Printf(TEXT("There are no samples in this track, cannot create an iterator."));
+				return nullptr;
+			}
+			return CreateIterator(Convs.NumTotalSamples - 1);
+		}
+
+		TSharedPtr<FMP4Track::FIterator> FMP4Track::CreateIterator(uint32 InAtSampleNumber)
+		{
+			TSharedPtr<FMP4Track::FIterator> It(CreateIterator());
+			if (It.IsValid() && InAtSampleNumber)
+			{
+				InAtSampleNumber = InAtSampleNumber < Convs.NumTotalSamples ? InAtSampleNumber : Convs.NumTotalSamples-1;
+				It->SampleNumber = InAtSampleNumber;
+				It->stszIt.SetToSampleNumber(InAtSampleNumber);
+				It->sttsIt.SetToSampleNumber(InAtSampleNumber);
+				It->cttsIt.SetToSampleNumber(InAtSampleNumber);
+				It->stcoIt.SetToSampleNumber(InAtSampleNumber);
+				It->stscIt.SetToSampleNumber(InAtSampleNumber);
+				It->stssIt.SetToSampleNumber(InAtSampleNumber);
+				It->rapIt.SetToSampleNumber(InAtSampleNumber);
+				It->Update();
+			}
+			return It;
+		}
+
+
+		TSharedPtr<FMP4Track::FIterator> FMP4Track::CreateIteratorAtKeyframe(FTimeValue InForTime, FTimeValue InLaterTimeThreshold)
+		{
+			check(bHasBeenPrepared);
+			if (!bHasBeenPrepared || !SttsBox.IsValid() || SttsBox->GetTotalDuration()==0)
+			{
+				LastErrorMessage = FString::Printf(TEXT("Track has not been prepared, cannot create an iterator."));
+				return nullptr;
+			}
+			if (!InForTime.IsValid())
+			{
+				LastErrorMessage = FString::Printf(TEXT("Invalid time, cannot create an iterator."));
+				return nullptr;
+			}
+			if (!Convs.DurationFromMdhdBox.IsValid() || !Convs.MappedDurationFromElstBox.IsValid())
+			{
+				LastErrorMessage = FString::Printf(TEXT("Invalid track duration, cannot create an iterator."));
+				return nullptr;
+			}
+			if (InForTime < FTimeValue::GetZero())
+			{
+				InForTime.SetToZero();
+			}
+			if (InLaterTimeThreshold < FTimeValue::GetZero())
+			{
+				InForTime.SetToZero();
+			}
+			const uint32 TrackTimescale = Convs.DurationFromMdhdBox.GetDenominator();
+			int64 LocalTrackTime = InForTime.GetAsTimebase(TrackTimescale);
+			// Clamp the time into the media time. The input may be larger than the media time, which is possible due to an
+			// edit list mapping more content into the timeline than the media has. We need to find the frame in the media
+			// though, so we clamp the time into the media time.
+			LocalTrackTime = LocalTrackTime > Convs.DurationFromMdhdBox.GetNumerator() ? Convs.DurationFromMdhdBox.GetNumerator() : LocalTrackTime;
+			int64 MaxLocalTrackTime = (InForTime + InLaterTimeThreshold).GetAsTimebase(TrackTimescale);
+
+			// Shift the search time into the media timeline.
+			LocalTrackTime += Convs.CompositionTimeAtZeroPoint;
+			MaxLocalTrackTime += Convs.CompositionTimeAtZeroPoint;
+
+			int64 ApproxSampleNumber = Convs.NumTotalSamples ? LocalTrackTime * (Convs.NumTotalSamples-1) / SttsBox->GetTotalDuration() : 0;
+			check(ApproxSampleNumber <= 0xffffffff);
+
+
+			TSharedPtr<FMP4Track::FIterator> ApproxFrameIt(CreateIterator((uint32)ApproxSampleNumber));
+			if (!ApproxFrameIt.IsValid())
+			{
+				LastErrorMessage = FString::Printf(TEXT("Failed to create track iterator for sample #%lld with %u samples in track"), (long long int)ApproxSampleNumber, Convs.NumTotalSamples);
+				return nullptr;
+			}
+			// Move the approximate iterator backwards or forwards towards the target time.
+			// This should not be off by much unless variable frame rate is used with greatly varying durations
+			// or an edit list cuts off significant amounts of the media.
+			if (ApproxFrameIt->GetPTS().GetNumerator() > LocalTrackTime)
+			{
+				for(; !ApproxFrameIt->IsFirst() && ApproxFrameIt->GetPTS().GetNumerator() > LocalTrackTime; ApproxFrameIt->Prev())
+				{ }
+			}
+			else if (ApproxFrameIt->GetPTS().GetNumerator()+ApproxFrameIt->GetDuration().GetNumerator() <= LocalTrackTime)
+			{
+				for(; !ApproxFrameIt->IsLast() && ApproxFrameIt->GetPTS().GetNumerator()+ApproxFrameIt->GetDuration().GetNumerator() <= LocalTrackTime; ApproxFrameIt->Next())
+				{ }
+			}
+			// Locate the nearest earlier sync sample, which might be the current one already.
+			TSharedPtr<FMP4Track::FIterator> SyncFrameIt(ApproxFrameIt->Clone());
+			for(; !SyncFrameIt->IsFirst() && (SyncFrameIt->GetPTS().GetNumerator() > LocalTrackTime || !SyncFrameIt->IsSyncOrRAPSample()); SyncFrameIt->Prev())
+			{ }
+			TSharedPtr<FMP4Track::FIterator> NextSyncFrameIt(ApproxFrameIt->Clone());
+			bool bLaterOneIsPossible = false;
+			if (MaxLocalTrackTime > LocalTrackTime)
+			{
+				// Due to possible frame reordering we need to look at the DTS here with the composition offset applied
+				// to be sure to find the correct sample. If we were to look at the PTS we could leave the loop to early.
+				for(int64 MaxDtsWithComp=MaxLocalTrackTime+Convs.CompositionTimeAtZeroPoint; !NextSyncFrameIt->IsLast() && (NextSyncFrameIt->GetDTS().GetNumerator() <= MaxDtsWithComp && !NextSyncFrameIt->IsSyncOrRAPSample()); NextSyncFrameIt->Next())
+				{ }
+				bLaterOneIsPossible = NextSyncFrameIt->IsSyncOrRAPSample() && NextSyncFrameIt->GetPTS().GetNumerator() <= MaxLocalTrackTime;
+			}
+			// Did we even find any sync sample?
+			if (!SyncFrameIt->IsSyncOrRAPSample() && !NextSyncFrameIt->IsSyncOrRAPSample())
+			{
+				LastErrorMessage = FString::Printf(TEXT("No sync sample found, cannot create an iterator."));
+				return nullptr;
+			}
+			// If there is a possible later one to use we need to check if the earlier one is outside the threshold.
+			if (bLaterOneIsPossible && LocalTrackTime - SyncFrameIt->GetPTS().GetNumerator() > MaxLocalTrackTime - LocalTrackTime)
+			{
+				return NextSyncFrameIt;
+			}
+			return SyncFrameIt;
+		}
+
+/****************************************************************************************************************************************************/
+/****************************************************************************************************************************************************/
+/****************************************************************************************************************************************************/
+
+		void FMP4Track::FIterator::Update()
+		{
+			if (!IsValid())
+			{
+				return;
+			}
+			int64 DTS = sttsIt.GetCurrentTime() + Convs.BaseMediaDecodeTime;
+			uint32 Duration = sttsIt.GetCurrentDuration();
+			int64 CompositionTimeOffset = cttsIt.GetCurrentOffset();
+			int64 PTS = DTS + CompositionTimeOffset;
+
+			const uint32 Timescale = Convs.DurationFromMdhdBox.GetDenominator();
+			CurrentDTS.SetFromND(DTS, Timescale);
+			CurrentPTS.SetFromND(PTS, Timescale);
+			CurrentEffectiveDTS.SetFromND(DTS - Convs.DTSShiftAtZeroPoint, Timescale);
+			CurrentEffectivePTS.SetFromND(PTS - Convs.CompositionTimeAtZeroPoint, Timescale);
+			CurrentDuration.SetFromND(Duration, Timescale);
+			CurrentSampleSize = stszIt.GetCurrentSampleSize();
+			bCurrentIsSyncOrRAP = stssIt.IsSyncSample() || rapIt.GetCurrentGroupDescriptionIndex() != 0;
+
+			// Which chunk is this sample in?
+			uint32 chunkIndex = stscIt.GetCurrentChunkIndex();
+			check(chunkIndex);
+			uint64 chunkOffset = stcoIt.GetOffsetForChunkIndex(chunkIndex - 1);
+			check(chunkOffset);
+			// Which sample position within the current chunk run are we at?
+			uint32 samplePosInChunk = stscIt.GetSampleIndexInCurrentChunk();
+			// Giving us which sample number at the start of the chunk?
+			uint32 sampleNumAtChunkStart = SampleNumber - samplePosInChunk;
+			for(uint32 i=0; i<samplePosInChunk; ++i)
+			{
+				chunkOffset += stszIt.GetSampleSizeForSampleNum(sampleNumAtChunkStart + i);
+			}
+			CurrentSampleFileOffset = (int64) chunkOffset;
+		}
+
+
+		bool FMP4Track::FIterator::Next()
+		{
+			if (SampleNumber+1 < Convs.NumTotalSamples)
+			{
+				// Note: stcoIt is not an iterator, so there's nothing to call on it.
+				verify(stszIt.Next());
+				verify(sttsIt.Next());
+				verify(cttsIt.Next());
+				verify(stscIt.Next());
+				verify(stssIt.Next());
+				verify(rapIt.Next());
+				++SampleNumber;
+				Update();
+				return true;
+			}
+			return false;
+		}
+		bool FMP4Track::FIterator::Prev()
+		{
+			if (SampleNumber > 0)
+			{
+				// Note: stcoIt is not an iterator, so there's nothing to call on it.
+				verify(stszIt.Prev());
+				verify(sttsIt.Prev());
+				verify(cttsIt.Prev());
+				verify(stscIt.Prev());
+				verify(stssIt.Prev());
+				verify(rapIt.Prev());
+				--SampleNumber;
+				Update();
+				return true;
+			}
+			return false;
+		}
+
+
+		bool FMP4Track::FIterator::NextEffective()
+		{
+			// The last sample number is inclusive, that is, that sample is needed.
+			if (SampleNumber+1 <= Convs.LastSample.LastSampleNumber)
+			{
+				// Note: stcoIt is not an iterator, so there's nothing to call on it.
+				verify(stszIt.Next());
+				verify(sttsIt.Next());
+				verify(cttsIt.Next());
+				verify(stscIt.Next());
+				verify(stssIt.Next());
+				verify(rapIt.Next());
+				++SampleNumber;
+				Update();
+				return true;
+			}
+			return false;
+		}
+		bool FMP4Track::FIterator::PrevEffective()
+		{
+			if (SampleNumber > Convs.FirstSample.SyncSampleNumber)
+			{
+				// Note: stcoIt is not an iterator, so there's nothing to call on it.
+				verify(stszIt.Prev());
+				verify(sttsIt.Prev());
+				verify(cttsIt.Prev());
+				verify(stscIt.Prev());
+				verify(stssIt.Prev());
+				verify(rapIt.Prev());
+				--SampleNumber;
+				Update();
+				return true;
+			}
+			return false;
+		}
+
+	} // namespace UtilitiesMP4
+
+} // namespace Electra
diff --git a/Engine/Plugins/Media/ElectraUtil/Source/ElectraBase/Private/Utilities/UtilitiesMP4.cpp b/Engine/Plugins/Media/ElectraUtil/Source/ElectraBase/Private/Utilities/UtilitiesMP4.cpp
new file mode 100644
--- /dev/null
+++ b/Engine/Plugins/Media/ElectraUtil/Source/ElectraBase/Private/Utilities/UtilitiesMP4.cpp
@@ -0,0 +1,560 @@
+// Copyright Epic Games, Inc. All Rights Reserved.
+
+#include "Utilities/UtilitiesMP4.h"
+#include "Utilities/MP4Boxes/MP4Boxes.h"
+
+#include "HAL/FileManager.h"
+#include "Misc/FileHelper.h"
+#include "Serialization/Archive.h"
+#include "Serialization/ArrayReader.h"
+
+
+namespace Electra
+{
+namespace UtilitiesMP4
+{
+
+
+FMP4AtomReaderHelper::FMP4AtomReaderHelper(const TConstArrayView<uint8>& InData)
+	: DataPtr(InData.GetData()), DataSize(InData.Num()), CurrentOffset(0)
+{
+}
+
+int64 FMP4AtomReaderHelper::GetCurrentOffset() const
+{
+	return CurrentOffset;
+}
+
+int64 FMP4AtomReaderHelper::GetNumBytesRemaining() const
+{
+	return DataSize - GetCurrentOffset();
+}
+
+const uint8* FMP4AtomReaderHelper::GetCurrentDataPointer() const
+{
+	return GetNumBytesRemaining() ? DataPtr + GetCurrentOffset() : nullptr;
+}
+
+void FMP4AtomReaderHelper::SetCurrentOffset(int64 InNewOffset)
+{
+	check(InNewOffset >= 0 && InNewOffset <= DataSize);
+	if (InNewOffset >= 0 && InNewOffset <= DataSize)
+	{
+		CurrentOffset = InNewOffset;
+	}
+}
+
+bool FMP4AtomReaderHelper::ReadVersionAndFlags(uint8& OutVersion, uint32& OutFlags)
+{
+	uint32 VersionAndFlags = 0;
+	if (!Read(VersionAndFlags))
+	{
+		return false;
+	}
+	OutVersion = (uint8)(VersionAndFlags >> 24);
+	OutFlags = VersionAndFlags & 0x00ffffff;
+	return true;
+}
+
+bool FMP4AtomReaderHelper::ReadString(FString& OutString, uint16 InNumBytes)
+{
+	OutString.Empty();
+	if (InNumBytes == 0)
+	{
+		return true;
+	}
+	TArray<uint8> Buf;
+	Buf.AddUninitialized(InNumBytes);
+	if (ReadBytes(Buf.GetData(), InNumBytes))
+	{
+		// Check for UTF16 BOM
+		if (InNumBytes >= 2 && ((Buf[0] == 0xff && Buf[1] == 0xfe) || (Buf[0] == 0xfe && Buf[1] == 0xff)))
+		{
+			// String uses UTF16, which is not supported
+			return false;
+		}
+		FUTF8ToTCHAR cnv((const ANSICHAR*)Buf.GetData(), InNumBytes);
+		OutString = FString::ConstructFromPtrSize(cnv.Get(), cnv.Length());
+		return true;
+	}
+	return false;
+}
+
+bool FMP4AtomReaderHelper::ReadStringUTF8(FString& OutString, int32 InNumBytes)
+{
+	OutString.Empty();
+	if (InNumBytes == 0)
+	{
+		return true;
+	}
+	else if (InNumBytes < 0)
+	{
+		InNumBytes = GetNumBytesRemaining();
+		check(InNumBytes >= 0);
+		if (InNumBytes < 0)
+		{
+			return false;
+		}
+	}
+	TArray<uint8> Buf;
+	Buf.AddUninitialized(InNumBytes);
+	if (ReadBytes(Buf.GetData(), InNumBytes))
+	{
+		FUTF8ToTCHAR cnv((const ANSICHAR*)Buf.GetData(), InNumBytes);
+		OutString = FString::ConstructFromPtrSize(cnv.Get(), cnv.Length());
+		return true;
+	}
+	return false;
+}
+
+bool FMP4AtomReaderHelper::ReadStringUTF16(FString& OutString, int32 InNumBytes)
+{
+	OutString.Empty();
+	if (InNumBytes == 0)
+	{
+		return true;
+	}
+	else if (InNumBytes < 0)
+	{
+		InNumBytes = GetNumBytesRemaining();
+		check(InNumBytes >= 0);
+		if (InNumBytes < 0)
+		{
+			return false;
+		}
+	}
+	TArray<uint8> Buf;
+	Buf.AddUninitialized(InNumBytes);
+	if (ReadBytes(Buf.GetData(), InNumBytes))
+	{
+		check(!"TODO");
+/*
+		FUTF8ToTCHAR cnv((const ANSICHAR*)Buf.GetData(), InNumBytes);
+		OutString = FString(cnv.Length(), cnv.Get());
+		return true;
+*/
+	}
+	return false;
+}
+
+bool FMP4AtomReaderHelper::ReadAsNumber(uint64& OutValue, int32 InNumBytes)
+{
+	OutValue = 0;
+	if (InNumBytes < 0 || InNumBytes > 8)
+	{
+		return false;
+	}
+	for(int32 i=0; i<InNumBytes; ++i)
+	{
+		uint8 d;
+		if (!Read(d))
+		{
+			return false;
+		}
+		OutValue = (OutValue << 8) | d;
+	}
+	return true;
+}
+bool FMP4AtomReaderHelper::ReadAsNumber(int64& OutValue, int32 InNumBytes)
+{
+	OutValue = 0;
+	if (InNumBytes < 0 || InNumBytes > 8)
+	{
+		return false;
+	}
+	for(int32 i=0; i<InNumBytes; ++i)
+	{
+		uint8 d;
+		if (!Read(d))
+		{
+			return false;
+		}
+		if (i==0 && d>127)
+		{
+			OutValue = -1;
+		}
+		OutValue = (OutValue << 8) | d;
+	}
+	return true;
+}
+bool FMP4AtomReaderHelper::ReadAsNumber(float& OutValue)
+{
+	uint32 Flt;
+	if (Read(Flt))
+	{
+		OutValue = *reinterpret_cast<float*>(&Flt);
+		return true;
+	}
+	return false;
+}
+bool FMP4AtomReaderHelper::ReadAsNumber(double& OutValue)
+{
+	uint64 Dbl;
+	if (Read(Dbl))
+	{
+		OutValue = *reinterpret_cast<double*>(&Dbl);
+		return true;
+	}
+	return false;
+}
+
+bool FMP4AtomReaderHelper::ReadBytes(void* Buffer, int32 InNumBytes)
+{
+	return ReadData(Buffer, InNumBytes) == InNumBytes;
+}
+
+int32 FMP4AtomReaderHelper::ReadData(void* IntoBuffer, int32 NumBytesToRead)
+{
+	if (NumBytesToRead <= 0)
+	{
+		return 0;
+	}
+	int64 NumAvail = DataSize - CurrentOffset;
+	if (NumAvail >= NumBytesToRead)
+	{
+		if (IntoBuffer)
+		{
+			FMemory::Memcpy(IntoBuffer, DataPtr + CurrentOffset, NumBytesToRead);
+		}
+		CurrentOffset += NumBytesToRead;
+		return NumBytesToRead;
+	}
+	return -1;
+}
+
+bool FMP4AtomReaderHelper::ParseIntoBoxInfo(FMP4BoxInfo& OutBoxInfo, int64 InAtFileOffset)
+{
+	// Clear output with default values before continuing.
+	OutBoxInfo = FMP4BoxInfo();
+	uint32 BoxSize, BoxType;
+	if (!Read(BoxSize) || !Read(BoxType))
+	{
+		return false;
+	}
+	OutBoxInfo.Offset = InAtFileOffset;
+	OutBoxInfo.Size = (int64) BoxSize;
+	OutBoxInfo.Type = BoxType;
+#if !UE_BUILD_SHIPPING
+	OutBoxInfo.Name[0] = (char) ((BoxType >> 24) & 255);
+	OutBoxInfo.Name[1] = (char) ((BoxType >> 16) & 255);
+	OutBoxInfo.Name[2] = (char) ((BoxType >>  8) & 255);
+	OutBoxInfo.Name[3] = (char) ((BoxType >>  0) & 255);
+	OutBoxInfo.Name[4] = 0;
+#endif
+	OutBoxInfo.DataOffset = 8;
+	// Check the box size value.
+	if (OutBoxInfo.Size == 1)
+	{
+		uint64 BoxSize64;
+		if (!Read(BoxSize64))
+		{
+			return false;
+		}
+		OutBoxInfo.DataOffset += 8;
+		OutBoxInfo.Size = (int64) BoxSize64;
+	}
+	// Is the box type a UUID ?
+	if (OutBoxInfo.Type == MakeBoxAtom('u','u','i','d'))
+	{
+		// Read additional 16 bytes for the UUID
+		if (ReadData(OutBoxInfo.UUID, 16) != 16)
+		{
+			return false;
+		}
+		OutBoxInfo.DataOffset += 16;
+	}
+	OutBoxInfo.Data = MakeConstArrayView(GetCurrentDataPointer(), OutBoxInfo.Size ? OutBoxInfo.Size - OutBoxInfo.DataOffset : GetNumBytesRemaining());
+	return true;
+}
+
+
+
+/**********************************************************************************************************************/
+/**********************************************************************************************************************/
+/**********************************************************************************************************************/
+
+
+
+FMP4BoxLocatorReader::~FMP4BoxLocatorReader()
+{
+}
+
+bool FMP4BoxLocatorReader::LocateAndReadRootBoxes(TArray<TSharedPtr<FMP4BoxData, ESPMode::ThreadSafe>>& OutBoxInfos, const TSharedPtr<IBaseDataReader, ESPMode::ThreadSafe>& InDataReader, const TArray<uint32>& InFirstBoxes, const TArray<uint32>& InStopAfterBoxes, const TArray<uint32>& InReadDataOfBoxes, IBaseDataReader::FCancellationCheckDelegate InCheckCancellationDelegate)
+{
+	// We NEVER want to read the `mdat` box here!
+	check(!InReadDataOfBoxes.Contains(MakeBoxAtom('m','d','a','t')));
+
+	check(InDataReader.IsValid());
+	if (!InDataReader.IsValid())
+	{
+		return false;
+	}
+	CurrentOffset = InDataReader->GetCurrentFileOffset();
+
+	#define CHECK_READ(NumReq) \
+		if (NumRead == IBaseDataReader::EResult::Canceled) \
+		{ return false;	} \
+		else if (NumRead == IBaseDataReader::EResult::ReadError) \
+		{ LastError = InDataReader->GetLastError(); return false; } \
+		else if (/*NumRead == IBaseDataReader::EResult::ReachedEOF ||*/ NumRead != NumReq) \
+		{ LastError = FString::Printf(TEXT("File truncated. Cannot read %lld bytes from offset %lld"), (long long int)NumReq, (long long int)CurrentOffset+BoxInternalOffset); return false; }
+
+	int64 TotalFileSize = -1;
+	for(int32 BoxNum=0; ;++BoxNum)
+	{
+		union UBuf
+		{
+			uint64 As64[2];
+			uint32 As32[4];
+			uint8 As8[16];
+		};
+		UBuf BoxSizeAndType;
+
+		// Read 8 bytes
+		uint32 BoxInternalOffset = 0;
+		int64 NumRead = InDataReader->ReadData(BoxSizeAndType.As64, 8, CurrentOffset, InCheckCancellationDelegate);
+		CHECK_READ(8);
+
+		BoxInternalOffset = 8;
+		TSharedPtr<FMP4BoxData, ESPMode::ThreadSafe> bi = MakeShared<FMP4BoxData, ESPMode::ThreadSafe>();
+		bi->Size = (int64) GetFromBigEndian(BoxSizeAndType.As32[0]);
+		bi->Offset = CurrentOffset;
+		bi->Type = GetFromBigEndian(BoxSizeAndType.As32[1]);
+#if !UE_BUILD_SHIPPING
+		bi->Name[0] = BoxSizeAndType.As8[4]; bi->Name[1] = BoxSizeAndType.As8[5]; bi->Name[2] = BoxSizeAndType.As8[6]; bi->Name[3] = BoxSizeAndType.As8[7]; bi->Name[4] = 0;
+#endif
+
+		// After having read the first few bytes we should now know the overall filesize.
+		if (BoxNum == 0)
+		{
+			TotalFileSize = InDataReader->GetTotalFileSize();
+			if (InFirstBoxes.Num() && !InFirstBoxes.Contains(bi->Type))
+			{
+				LastError = TEXT("Invalid mp4 file: First box is not of expected type");
+				return false;
+			}
+		}
+
+		// Check the box size value.
+		if (bi->Size == 0)
+		{
+			// Zero size means "until the end of the file".
+			bi->Size = TotalFileSize > 0 ? TotalFileSize - CurrentOffset : -1;
+		}
+		else if (bi->Size == 1)
+		{
+			// A size of 1 indicates that the size is expressed as a 64 bit value following the box type.
+			// Read additional 8 bytes for the 64 bit length
+			NumRead = InDataReader->ReadData(BoxSizeAndType.As64, 8, CurrentOffset+BoxInternalOffset, InCheckCancellationDelegate);
+			CHECK_READ(8);
+			BoxInternalOffset += 8;
+			bi->Size = (int64) GetFromBigEndian(BoxSizeAndType.As64[0]);
+		}
+
+		// Is the box type a UUID ?
+		if (bi->Type == MakeBoxAtom('u','u','i','d'))
+		{
+			// Read additional 16 bytes for the UUID
+			NumRead = InDataReader->ReadData(bi->UUID, 16, CurrentOffset+BoxInternalOffset, InCheckCancellationDelegate);
+			CHECK_READ(16);
+			BoxInternalOffset += 16;
+		}
+		// Shall we read this box?
+		if ((bi->Type != MakeBoxAtom('m','d','a','t')) && (InReadDataOfBoxes.IsEmpty() || InReadDataOfBoxes.Contains(bi->Type)))
+		{
+			bi->DataBuffer.SetNumUninitialized(bi->Size-BoxInternalOffset);
+			bi->Data = MakeConstArrayView(bi->DataBuffer);
+			NumRead = InDataReader->ReadData(bi->DataBuffer.GetData(), bi->Size-BoxInternalOffset, CurrentOffset+BoxInternalOffset, InCheckCancellationDelegate);
+		}
+		else
+		{
+			NumRead = InDataReader->ReadData(nullptr, bi->Size-BoxInternalOffset, CurrentOffset+BoxInternalOffset, InCheckCancellationDelegate);
+		}
+		CHECK_READ(bi->Size - BoxInternalOffset);
+		bi->DataOffset = BoxInternalOffset;
+
+		// Advance the current offset, whether we have read the box or not.
+		CurrentOffset += bi->Size;
+		check(CurrentOffset == InDataReader->GetCurrentFileOffset());
+		bool bStopNow = InStopAfterBoxes.Contains(bi->Type);
+		OutBoxInfos.Emplace(MoveTemp(bi));
+		if (bStopNow || InDataReader->HasReachedEOF())
+		{
+			return true;
+		}
+	}
+	return false;
+}
+
+
+
+
+
+bool FMP4BoxTreeParser::ParseBoxTreeInternal(const TWeakPtr<FMP4BoxBase>& InParent, const FMP4BoxInfo& InBox)
+{
+	BoxTree = FMP4BoxFactory::Get().Create(InParent, InBox);
+	// Parse the enclosed box recursively unless it contains
+	// a list of entries that only the box knows how to parse.
+	if (BoxTree.IsValid() && !BoxTree->IsLeafBox() && !BoxTree->IsListOfEntries())
+	{
+		// The data of this container box represents one or several other boxes.
+		// We need to parse them one by one until there is no more data here.
+		const TConstArrayView<uint8>& bd(BoxTree->GetBoxData());
+		const uint8* BoxData = bd.GetData();
+		int64 BoxBytesRemaining = bd.Num();
+		int64 NextBoxOffset = BoxTree->GetBoxFileOffset() + BoxTree->GetBoxDataOffset();
+		while(BoxBytesRemaining > 0)
+		{
+			FMP4AtomReaderHelper bh(MakeConstArrayView(BoxData, BoxBytesRemaining));
+			FMP4BoxInfo bi;
+			if (!bh.ParseIntoBoxInfo(bi, NextBoxOffset))
+			{
+				return false;
+			}
+			FMP4BoxTreeParser bp;
+			if (!bp.ParseBoxTreeInternal(BoxTree, bi))
+			{
+				return false;
+			}
+			BoxTree->AddChildBox(MoveTemp(bp.BoxTree));
+			BoxData += bi.Size;
+			BoxBytesRemaining -= bi.Size;
+			NextBoxOffset = bi.Offset + bi.Size;
+		}
+	}
+	return true;
+}
+
+bool FMP4BoxTreeParser::ParseBoxTree(const TSharedPtr<FMP4BoxInfo, ESPMode::ThreadSafe>& InRootBox)
+{
+	check(InRootBox.IsValid());
+	bool bOk = ParseBoxTreeInternal(nullptr, *InRootBox);
+	if (bOk && BoxTree.IsValid())
+	{
+		BoxTree->SetRootBoxData(InRootBox);
+	}
+	return bOk;
+}
+
+
+
+} // namespace UtilitiesMP4
+
+
+
+/**********************************************************************************************************************/
+/**********************************************************************************************************************/
+/**********************************************************************************************************************/
+
+class FFileDataReader : public IFileDataReader
+{
+public:
+	~FFileDataReader();
+
+	bool Open(const FString& InFilename) override;
+
+	//-------------------------------------------------------------------------
+	// Methods from IBaseDataReader
+	//
+	int64 ReadData(void* InOutBuffer, int64 InNumBytes, int64 InFromOffset, FCancellationCheckDelegate InCheckCancellationDelegate) override;
+	int64 GetTotalFileSize() override;
+	int64 GetCurrentFileOffset() override;
+	bool HasReachedEOF() override;
+	FString GetLastError() override;
+private:
+	void Close();
+	FString LastError;
+	FArchive* Archive = nullptr;
+	int64 TotalFileSize = -1;
+	int64 CurrentOffset = 0;
+};
+
+TSharedPtr<IFileDataReader, ESPMode::ThreadSafe> IFileDataReader::Create()
+{
+	return MakeShared<FFileDataReader, ESPMode::ThreadSafe>();
+}
+
+FFileDataReader::~FFileDataReader()
+{
+	Close();
+}
+
+void FFileDataReader::Close()
+{
+	if (Archive)
+	{
+		Archive->Close();
+		delete Archive;
+		Archive = nullptr;
+	}
+}
+
+bool FFileDataReader::Open(const FString& InFilename)
+{
+	Archive = IFileManager::Get().CreateFileReader(*InFilename, 0);
+	if (!Archive)
+	{
+		LastError = FString::Printf(TEXT("Failed to open file \"%s\""), *InFilename);
+		return false;
+	}
+	TotalFileSize = Archive->TotalSize();
+	return true;
+}
+
+int64 FFileDataReader::ReadData(void* InOutBuffer, int64 InNumBytes, int64 InFromOffset, FCancellationCheckDelegate InCheckCancellationDelegate)
+{
+	check(InNumBytes >= 0);
+	if (!Archive)
+	{
+		LastError = FString::Printf(TEXT("File reader has not been openend"));
+		return EResult::ReadError;
+	}
+	if (InNumBytes <= 0)
+	{
+		return 0;
+	}
+	if (CurrentOffset != InFromOffset)
+	{
+		check(InFromOffset >= 0 && InFromOffset <= TotalFileSize);
+		CurrentOffset = InFromOffset >= 0 ? InFromOffset <= TotalFileSize ? InFromOffset : TotalFileSize : 0;
+		Archive->Seek(CurrentOffset);
+	}
+	if (InNumBytes + CurrentOffset > TotalFileSize)
+	{
+		InNumBytes = TotalFileSize - CurrentOffset;
+	}
+	if (InOutBuffer)
+	{
+		Archive->Serialize(InOutBuffer, InNumBytes);
+	}
+	else
+	{
+		Archive->Seek(CurrentOffset + InNumBytes);
+	}
+	CurrentOffset += InNumBytes;
+	return InNumBytes;
+}
+
+int64 FFileDataReader::GetTotalFileSize()
+{
+	return TotalFileSize;
+}
+
+int64 FFileDataReader::GetCurrentFileOffset()
+{
+	return CurrentOffset;
+}
+
+bool FFileDataReader::HasReachedEOF()
+{
+	return CurrentOffset >= TotalFileSize;
+}
+
+FString FFileDataReader::GetLastError()
+{
+	return LastError;
+}
+
+
+
+} // namespace Electra
diff --git a/Engine/Plugins/Media/ElectraUtil/Source/ElectraBase/Public/Utilities/IElectraBaseDataReader.h b/Engine/Plugins/Media/ElectraUtil/Source/ElectraBase/Public/Utilities/IElectraBaseDataReader.h
new file mode 100644
--- /dev/null
+++ b/Engine/Plugins/Media/ElectraUtil/Source/ElectraBase/Public/Utilities/IElectraBaseDataReader.h
@@ -0,0 +1,75 @@
+// Copyright Epic Games, Inc. All Rights Reserved.
+
+#pragma once
+
+#include "CoreTypes.h"
+#include "Containers/UnrealString.h"
+#include "Delegates/Delegate.h"
+
+
+namespace Electra
+{
+	/**
+	 * Abstract base class to pass into
+	 */
+	class IBaseDataReader
+	{
+	public:
+		virtual ~IBaseDataReader() = default;
+
+		enum EResult : int64
+		{
+			ReadError = -1,
+			ReachedEOF = -2,
+			Canceled = -3
+		};
+
+		/**
+		 * A delegate passed into the read methods through which the reader implementation
+		 * can call into to see if the read request has been canceled.
+		 */
+		DECLARE_DELEGATE_RetVal(bool, FCancellationCheckDelegate);
+
+		/**
+		 * This method reads data into the provided buffer from the specified absolute offset.
+		 * The number of bytes to read should be retrieved unless the end of the file is
+		 * reached where returning fewer bytes than requested is permitted.
+		 * Reading more bytes than requested is forbidden since the read buffer may not be
+		 * large enough to accommodate more bytes than asked for.
+		 * Negative return values to indicate a problem are defined in `EResult`
+		 * `EResult::Canceled` may be returned even if the provided cancellation delegate
+		 * does not indicate cancelation, but the reader implementation has been canceled
+		 * by other means (like application shutdown, or in case of an implementation reading
+		 * from the network some other condition).
+		 */
+		virtual int64 ReadData(void* InOutBuffer, int64 InNumBytes, int64 InFromOffset, FCancellationCheckDelegate InCheckCancellationDelegate) = 0;
+
+		/**
+		 * This method shall return the total size of the file.
+		 * If the size is only known after performing the first read, -1 may be returned until the size is known.
+		 * If the file is unbounded -1 may be returned at all times.
+		 */
+		virtual int64 GetTotalFileSize() = 0;
+
+		/**
+		 * This method shall return the current file offset, which is initially zero unless the file
+		 * has been opened such that the initial position for this reader is not zero.
+		 * The return value is the absolute file position the next read would occur at.
+		 */
+		virtual int64 GetCurrentFileOffset() = 0;
+
+		/**
+		 * A convenience method to return `true` when all data has been read.
+		 * There is no requirement that `GetTotalFileSize()` returns a positive value.
+		 * This method must return `true` when `ReadData()` would return `EResult::ReachedEOF`.
+		 */
+		virtual bool HasReachedEOF() = 0;
+
+		/**
+		 * If a read error occurred and `ReadData()` returns `EResult::ReadError` a human-readable
+		 * message of what caused the error should be returned.
+		 */
+		virtual FString GetLastError() = 0;
+	};
+
+} // namespace Electra
diff --git a/Engine/Plugins/Media/ElectraUtil/Source/ElectraBase/Public/Utilities/MP4Boxes/MP4BoxBase.h b/Engine/Plugins/Media/ElectraUtil/Source/ElectraBase/Public/Utilities/MP4Boxes/MP4BoxBase.h
new file mode 100644
--- /dev/null
+++ b/Engine/Plugins/Media/ElectraUtil/Source/ElectraBase/Public/Utilities/MP4Boxes/MP4BoxBase.h
@@ -0,0 +1,131 @@
+// Copyright Epic Games, Inc. All Rights Reserved.
+
+#pragma once
+
+#include "Utilities/UtilitiesMP4.h"
+
+namespace Electra
+{
+	namespace UtilitiesMP4
+	{
+		class FMP4BoxBase : public TSharedFromThis<FMP4BoxBase>
+		{
+		public:
+			static TSharedPtr<FMP4BoxBase> Create(const TWeakPtr<FMP4BoxBase>& InParent, const FMP4BoxInfo& InBoxInfo)
+			{ return MakeShareable<FMP4BoxBase>(new FMP4BoxBase(InParent, InBoxInfo)); }
+			virtual ~FMP4BoxBase() = default;
+
+			virtual const FMP4BoxInfo& GetBoxInfo() const
+			{ return BoxInfo; }
+			virtual uint32 GetType() const
+			{ return BoxInfo.Type; }
+			virtual int64 GetBoxSize() const
+			{ return BoxInfo.Size; }
+			virtual TConstArrayView<uint8> GetBoxData() const
+			{ return BoxInfo.Data; }
+			virtual int64 GetBoxFileOffset() const
+			{ return BoxInfo.Offset; }
+			virtual int64 GetBoxDataOffset() const
+			{ return BoxInfo.DataOffset; }
+			virtual bool IsLeafBox() const
+			{ return true; }
+			virtual bool IsListOfEntries() const
+			{ return false; }
+			virtual bool IsSampleDescription() const
+			{ return false; }
+
+			//-------------------------------------------------------------------------
+			// Parent box query methods
+			//
+			template<typename T>
+			TSharedPtr<T> FindParentBox(uint32 InType) const
+			{
+				TSharedPtr<T> Parent = StaticCastSharedPtr<T>(ParentBox.Pin());
+				while(Parent.IsValid())
+				{
+					if (Parent->GetType() == InType)
+					{
+						return Parent;
+					}
+					Parent = StaticCastSharedPtr<T>(Parent->ParentBox.Pin());
+				}
+				return Parent;
+			}
+
+			//-------------------------------------------------------------------------
+			// Child box query methods
+			//
+			const TArray<TSharedPtr<FMP4BoxBase>> GetChildren() const
+			{ return Children; }
+
+			// Find the first box of the given type being a child of this box, or any child of a child box.
+			template<typename T>
+			TSharedPtr<T> FindBoxRecursive(uint32 InType, int32 InMaxDepth=8) const
+			{
+				// First pass, check all children if they are of the type
+				for(int32 i=0,iMax=Children.Num(); i<iMax; ++i)
+				{
+					if (Children[i]->GetType() == InType)
+					{
+						return StaticCastSharedPtr<T>(Children[i]);
+					}
+				}
+				// Second pass, check children recursively
+				if (InMaxDepth > 0)
+				{
+					for(int32 i=0, iMax=Children.Num(); i<iMax; ++i)
+					{
+						TSharedPtr<T> Box = Children[i]->FindBoxRecursive<T>(InType, InMaxDepth - 1);
+						if (Box.IsValid())
+						{
+							return Box;
+						}
+					}
+				}
+				return nullptr;
+			}
+
+
+			// Returns all instances of a given box type from the direct children of THIS BOX ONLY.
+			template<typename T>
+			void GetAllBoxInstances(TArray<TSharedPtr<T>>& OutAllBoxes, uint32 InType) const
+			{
+				for(int32 i=0,iMax=Children.Num(); i<iMax; ++i)
+				{
+					if (Children[i]->GetType() == InType)
+					{
+						OutAllBoxes.Emplace(StaticCastSharedPtr<T>(Children[i]));
+					}
+				}
+			}
+
+
+
+			// Called by building the box tree. Not to be called by user code.
+			virtual bool AddChildBox(TSharedPtr<FMP4BoxBase> InChildBox)
+			{
+				if (InChildBox.IsValid())
+				{
+					Children.Emplace(MoveTemp(InChildBox));
+					return true;
+				}
+				return false;
+			}
+
+			virtual void ProcessBoxChildrenRecursively(FMP4AtomReaderHelper& InOutReader, const FMP4BoxInfo& InCurrentBoxInfo);
+
+			void SetRootBoxData(const TSharedPtr<FMP4BoxInfo, ESPMode::ThreadSafe>& InRootBoxData)
+			{ RootBoxData = InRootBoxData; }
+		protected:
+			FMP4BoxBase(const TWeakPtr<FMP4BoxBase>& InParent, const FMP4BoxInfo& InBoxInfo) : ParentBox(InParent), BoxInfo(InBoxInfo)
+			{ }
+
+			TSharedPtr<FMP4BoxInfo, ESPMode::ThreadSafe> RootBoxData;
+			TWeakPtr<FMP4BoxBase> ParentBox;
+			TArray<TSharedPtr<FMP4BoxBase>> Children;
+			FMP4BoxInfo BoxInfo;
+		};
+
+	} // namespace UtilitiesMP4
+
+} // namespace Electra
diff --git a/Engine/Plugins/Media/ElectraUtil/Source/ElectraBase/Public/Utilities/MP4Boxes/MP4BoxIterators.h b/Engine/Plugins/Media/ElectraUtil/Source/ElectraBase/Public/Utilities/MP4Boxes/MP4BoxIterators.h
new file mode 100644
--- /dev/null
+++ b/Engine/Plugins/Media/ElectraUtil/Source/ElectraBase/Public/Utilities/MP4Boxes/MP4BoxIterators.h
@@ -0,0 +1,752 @@
+// Copyright Epic Games, Inc. All Rights Reserved.
+
+#pragma once
+
+#include "Utilities/MP4Boxes/MP4Boxes.h"
+
+namespace Electra
+{
+	namespace UtilitiesMP4
+	{
+
+		/****************************************************************************************************************************************************/
+
+		struct FSTTSBoxIterator
+		{
+			void SetBox(const TSharedPtr<FMP4BoxSTTS>& InSttsBox)
+			{
+				SttsBox = InSttsBox;
+				NumTotalSamples = SttsBox->GetNumTotalSamples();
+				Time = 0;
+				CurrentSampleNum = 0;
+				LocalPosInCurrent = 0;
+				CurrentDuration = NumTotalSamples ? SttsBox->GetEntries()[0].sample_delta : 0;
+				CurrentEntryIndex = 0;
+			}
+
+			void SetToSampleNumber(uint32 InSampleNum)
+			{
+				check(SttsBox.IsValid());
+				check(InSampleNum < NumTotalSamples);
+				InSampleNum = InSampleNum < NumTotalSamples ? InSampleNum : NumTotalSamples - 1;
+				const TArray<FMP4BoxSTTS::FEntry>& Entries(SttsBox->GetEntries());
+				Time = 0;
+				CurrentEntryIndex = 0;
+				LocalPosInCurrent = 0;
+				CurrentDuration = Entries[0].sample_delta;
+				uint32 n = 0;
+				// If the entire current entry is still before the wanted position we can skip over it.
+				while(n + Entries[CurrentEntryIndex].sample_count <= InSampleNum)
+				{
+					n += Entries[CurrentEntryIndex].sample_count;
+					Time += Entries[CurrentEntryIndex].sample_count * (int64)CurrentDuration;
+					++CurrentEntryIndex;
+					CurrentDuration = Entries[CurrentEntryIndex].sample_delta;
+				}
+				LocalPosInCurrent = InSampleNum - n;
+				Time += LocalPosInCurrent * (int64)CurrentDuration;
+				CurrentSampleNum = InSampleNum;
+			}
+
+			bool Next()
+			{
+				check(SttsBox.IsValid());
+				if (IsLast())
+				{
+					return false;
+				}
+				Time += CurrentDuration;
+				++CurrentSampleNum;
+				// Does the current range still provide?
+				const TArray<FMP4BoxSTTS::FEntry>& Entries(SttsBox->GetEntries());
+				if (++LocalPosInCurrent >= Entries[CurrentEntryIndex].sample_count)
+				{
+					// No, continue with the next entry.
+					++CurrentEntryIndex;
+					CurrentDuration = Entries[CurrentEntryIndex].sample_delta;
+					LocalPosInCurrent = 0;
+				}
+				return true;
+			}
+
+			bool Prev()
+			{
+				check(SttsBox.IsValid());
+				if (IsFirst())
+				{
+					return false;
+				}
+				// Need to go back one entry?
+				if (LocalPosInCurrent > 0)
+				{
+					// Not yet.
+					--LocalPosInCurrent;
+				}
+				else
+				{
+					// Yes.
+					const TArray<FMP4BoxSTTS::FEntry>& Entries(SttsBox->GetEntries());
+					--CurrentEntryIndex;
+					LocalPosInCurrent = Entries[CurrentEntryIndex].sample_count - 1;
+					CurrentDuration = Entries[CurrentEntryIndex].sample_delta;
+				}
+				Time -= CurrentDuration;
+				--CurrentSampleNum;
+				return true;
+			}
+
+			int64 GetCurrentTime() const
+			{ return Time; }
+			uint32 GetCurrentDuration() const
+			{ return CurrentDuration; }
+			uint32 GetCurrentSampleNum() const
+			{ return CurrentSampleNum; }
+			uint32 GetNumTotalSamples() const
+			{ return NumTotalSamples; }
+			bool IsFirst() const
+			{ return CurrentSampleNum == 0; }
+			bool IsLast() const
+			{ return CurrentSampleNum + 1 >= NumTotalSamples; }
+
+		private:
+			TSharedPtr<FMP4BoxSTTS> SttsBox;
+			int64 Time = 0;
+			uint32 CurrentSampleNum = 0;
+			uint32 NumTotalSamples = 0;
+			uint32 LocalPosInCurrent = 0;
+			uint32 CurrentDuration = 0;
+			int32 CurrentEntryIndex = 0;
+		};
+
+		/****************************************************************************************************************************************************/
+
+		struct FCTTSBoxIterator
+		{
+			void SetBox(const TSharedPtr<FMP4BoxCTTS>& InCttsBox, uint32 InMaxSamples)
+			{
+				if (InCttsBox.IsValid())
+				{
+					CttsBox = InCttsBox;
+					NumTotalSamples = CttsBox->GetNumTotalSamples();
+					CurrentOffset = NumTotalSamples ? CttsBox->GetEntries()[0].sample_offset : 0;
+				}
+				else
+				{
+					NumTotalSamples = InMaxSamples;
+					CurrentOffset = 0;
+				}
+				CurrentSampleNum = 0;
+				LocalPosInCurrent = 0;
+				CurrentEntryIndex = 0;
+			}
+
+			void SetToSampleNumber(uint32 InSampleNum)
+			{
+				check(InSampleNum < NumTotalSamples);
+				InSampleNum = InSampleNum < NumTotalSamples ? InSampleNum : NumTotalSamples - 1;
+				if (CttsBox.IsValid())
+				{
+					const TArray<FMP4BoxCTTS::FEntry>& Entries(CttsBox->GetEntries());
+					CurrentEntryIndex = 0;
+					LocalPosInCurrent = 0;
+					CurrentOffset = Entries[0].sample_offset;
+					uint32 n = 0;
+					// If the entire current entry is still before the wanted position we can skip over it.
+					while(n + Entries[CurrentEntryIndex].sample_count <= InSampleNum)
+					{
+						n += Entries[CurrentEntryIndex].sample_count;
+						++CurrentEntryIndex;
+						CurrentOffset = Entries[CurrentEntryIndex].sample_offset;
+					}
+					LocalPosInCurrent = InSampleNum - n;
+				}
+				CurrentSampleNum = InSampleNum;
+			}
+
+			bool Next()
+			{
+				if (IsLast())
+				{
+					return false;
+				}
+				if (CttsBox.IsValid())
+				{
+					// Does the current range still provide?
+					const TArray<FMP4BoxCTTS::FEntry>& Entries(CttsBox->GetEntries());
+					if (++LocalPosInCurrent >= Entries[CurrentEntryIndex].sample_count)
+					{
+						// No, continue with the next entry.
+						++CurrentEntryIndex;
+						CurrentOffset = Entries[CurrentEntryIndex].sample_offset;
+						LocalPosInCurrent = 0;
+					}
+				}
+				++CurrentSampleNum;
+				return true;
+			}
+
+			bool Prev()
+			{
+				if (IsFirst())
+				{
+					return false;
+				}
+				if (CttsBox.IsValid())
+				{
+					// Need to go back one entry?
+					if (LocalPosInCurrent > 0)
+					{
+						// Not yet.
+						--LocalPosInCurrent;
+					}
+					else
+					{
+						// Yes.
+						const TArray<FMP4BoxCTTS::FEntry>& Entries(CttsBox->GetEntries());
+						--CurrentEntryIndex;
+						LocalPosInCurrent = Entries[CurrentEntryIndex].sample_count - 1;
+						CurrentOffset = Entries[CurrentEntryIndex].sample_offset;
+					}
+				}
+				--CurrentSampleNum;
+				return true;
+			}
+
+			int64 GetCurrentOffset() const
+			{ return CurrentOffset; }
+			uint32 GetCurrentSampleNum() const
+			{ return CurrentSampleNum; }
+			bool IsFirst() const
+			{ return CurrentSampleNum == 0; }
+			bool IsLast() const
+			{ return CurrentSampleNum + 1 >= NumTotalSamples; }
+
+		private:
+			TSharedPtr<FMP4BoxCTTS> CttsBox;
+			int64 CurrentOffset = 0;
+			uint32 CurrentSampleNum = 0;
+			uint32 NumTotalSamples = 0;
+			uint32 LocalPosInCurrent = 0;
+			int32 CurrentEntryIndex = 0;
+		};
+
+		/****************************************************************************************************************************************************/
+
+		struct FSTSCBoxIterator
+		{
+			void SetBox(const TSharedPtr<FMP4BoxSTSC>& InStscBox, uint32 InMaxSamples)
+			{
+				StscBox = InStscBox;
+				NumTotalSamples = InMaxSamples;
+				CurrentSampleNum = 0;
+
+				if (StscBox->GetEntries().Num())
+				{
+					const TArray<FMP4BoxSTSC::FEntry>& Entries(StscBox->GetEntries());
+					Current_first_chunk = Entries[0].first_chunk;
+					Current_samples_per_chunk = Entries[0].samples_per_chunk;
+					Current_sample_description_index = Entries[0].sample_description_index;
+				}
+				else
+				{
+					Current_first_chunk = 0;
+					Current_samples_per_chunk = 0;
+				}
+				CurrentEntryIndex = 0;
+				LocalSampleInChunk = 0;
+			}
+
+			void SetToSampleNumber(uint32 InSampleNum)
+			{
+				if (StscBox.IsValid() && StscBox->GetEntries().Num())
+				{
+					InSampleNum = InSampleNum < NumTotalSamples ? InSampleNum : NumTotalSamples-1;
+					const TArray<FMP4BoxSTSC::FEntry>& Entries(StscBox->GetEntries());
+					const uint32 MaxEntries = Entries.Num();
+					uint32 n=0;
+					CurrentEntryIndex = 0;
+					Current_first_chunk = Entries[0].first_chunk;
+					Current_samples_per_chunk = Entries[0].samples_per_chunk;
+					Current_sample_description_index = Entries[0].sample_description_index;
+					while(1)
+					{
+						n += Current_samples_per_chunk;
+						if (n > InSampleNum)
+						{
+							LocalSampleInChunk = Current_samples_per_chunk - (n - InSampleNum);
+							break;
+						}
+						else
+						{
+							++Current_first_chunk;
+							if (CurrentEntryIndex + 1 < MaxEntries && Current_first_chunk == Entries[CurrentEntryIndex + 1].first_chunk)
+							{
+								++CurrentEntryIndex;
+								Current_first_chunk = Entries[CurrentEntryIndex].first_chunk;
+								Current_samples_per_chunk = Entries[CurrentEntryIndex].samples_per_chunk;
+								Current_sample_description_index = Entries[CurrentEntryIndex].sample_description_index;
+							}
+						}
+					}
+					CurrentSampleNum = InSampleNum;
+				}
+			}
+
+			bool Next()
+			{
+				if (StscBox.IsValid())
+				{
+					if (IsLast())
+					{
+						return false;
+					}
+
+					++CurrentSampleNum;
+					// Reached end of the current chunk run?
+					if (++LocalSampleInChunk == Current_samples_per_chunk)
+					{
+						const TArray<FMP4BoxSTSC::FEntry>& Entries(StscBox->GetEntries());
+						const uint32 MaxEntries = Entries.Num();
+						LocalSampleInChunk = 0;
+						++Current_first_chunk;
+						if (CurrentEntryIndex + 1 < MaxEntries && Current_first_chunk == Entries[CurrentEntryIndex + 1].first_chunk)
+						{
+							++CurrentEntryIndex;
+							Current_first_chunk = Entries[CurrentEntryIndex].first_chunk;
+							Current_samples_per_chunk = Entries[CurrentEntryIndex].samples_per_chunk;
+							Current_sample_description_index = Entries[CurrentEntryIndex].sample_description_index;
+						}
+					}
+					return true;
+				}
+				return false;
+			}
+
+			bool Prev()
+			{
+				if (StscBox.IsValid())
+				{
+					if (IsFirst())
+					{
+						return false;
+					}
+
+					--CurrentSampleNum;
+					// At the start of the current chunk run?
+					if (LocalSampleInChunk > 0)
+					{
+						--LocalSampleInChunk;
+					}
+					else
+					{
+						const TArray<FMP4BoxSTSC::FEntry>& Entries(StscBox->GetEntries());
+						check(Current_first_chunk);
+						// Are we at the start of this chunk run?
+						if (Current_first_chunk == Entries[CurrentEntryIndex].first_chunk)
+						{
+							// We need to go back an entry if there is one.
+							if (CurrentEntryIndex)
+							{
+								--CurrentEntryIndex;
+								Current_samples_per_chunk = Entries[CurrentEntryIndex].samples_per_chunk;
+								Current_sample_description_index = Entries[CurrentEntryIndex].sample_description_index;
+							}
+						}
+						--Current_first_chunk;
+						LocalSampleInChunk = Current_samples_per_chunk - 1;
+					}
+					return true;
+				}
+				return false;
+			}
+
+			// Note: The chunk index is 1-based, so you need to subtract 1 to use as an index into the `stco` box.
+			uint32 GetCurrentChunkIndex() const
+			{ return Current_first_chunk; }
+			uint32 GetNumSamplesInCurrentChunk() const
+			{ return Current_samples_per_chunk; }
+			uint32 GetSampleIndexInCurrentChunk() const
+			{ return LocalSampleInChunk; }
+			uint32 GetCurrentSampleNum() const
+			{ return CurrentSampleNum; }
+			bool IsFirst() const
+			{ return CurrentSampleNum == 0; }
+			bool IsLast() const
+			{ return CurrentSampleNum + 1 >= NumTotalSamples; }
+
+		private:
+			TSharedPtr<FMP4BoxSTSC> StscBox;
+			uint32 Current_first_chunk = 0;
+			uint32 Current_samples_per_chunk = 0;
+			uint32 Current_sample_description_index = 0;
+
+			uint32 CurrentSampleNum = 0;
+			uint32 NumTotalSamples = 0;
+			uint32 LocalSampleInChunk = 0;
+			uint32 CurrentEntryIndex = 0;
+		};
+
+		/****************************************************************************************************************************************************/
+
+		struct FSTSZBoxIterator
+		{
+			void SetBox(const TSharedPtr<FMP4BoxSTSZ>& InStszBox)
+			{
+				check(InStszBox.IsValid());
+				StszBox = InStszBox;
+				NumTotalSamples = StszBox->GetNumberOfSamples();
+				CurrentSampleNum = 0;
+			}
+
+			void SetToSampleNumber(uint32 InSampleNum)
+			{
+				CurrentSampleNum = InSampleNum < NumTotalSamples ? InSampleNum : NumTotalSamples-1;
+			}
+
+			bool Next()
+			{
+				if (StszBox.IsValid())
+				{
+					if (IsLast())
+					{
+						return false;
+					}
+					++CurrentSampleNum;
+					return true;
+				}
+				return false;
+			}
+
+			bool Prev()
+			{
+				if (StszBox.IsValid())
+				{
+					if (IsFirst())
+					{
+						return false;
+					}
+					--CurrentSampleNum;
+					return true;
+				}
+				return false;
+			}
+
+			uint32 GetCurrentSampleSize() const
+			{
+				check(StszBox.IsValid());
+				return StszBox->GetSizeOfSample(CurrentSampleNum);
+			}
+
+			uint32 GetSampleSizeForSampleNum(uint32 InForSampleNum) const
+			{
+				check(StszBox.IsValid());
+				return StszBox->GetSizeOfSample(InForSampleNum);
+			}
+
+			uint32 GetNumTotalSamples() const
+			{ return NumTotalSamples; }
+			uint32 GetCurrentSampleNum() const
+			{ return CurrentSampleNum; }
+			bool IsFirst() const
+			{ return CurrentSampleNum == 0; }
+			bool IsLast() const
+			{ return CurrentSampleNum + 1 >= NumTotalSamples; }
+
+		private:
+			TSharedPtr<FMP4BoxSTSZ> StszBox;
+			uint32 NumTotalSamples = 0;
+			uint32 CurrentSampleNum = 0;
+		};
+
+		/****************************************************************************************************************************************************/
+
+		// This is not so much an iterator than a holder of the box.
+		struct FSTCOBoxIterator
+		{
+			void SetBox(const TSharedPtr<FMP4BoxSTCO>& InStcoBox)
+			{
+				check(InStcoBox.IsValid());
+				StcoBox = InStcoBox;
+			}
+
+			void SetToSampleNumber(uint32 InSampleNum)
+			{
+			}
+
+			uint64 GetOffsetForChunkIndex(uint32 InChunkIndex) const
+			{
+				return StcoBox.IsValid() ? StcoBox->GetChunkOffset(InChunkIndex) : 0;
+			}
+		private:
+			TSharedPtr<FMP4BoxSTCO> StcoBox;
+		};
+
+		/****************************************************************************************************************************************************/
+
+		struct FSTSSBoxIterator
+		{
+			void SetBox(const TSharedPtr<FMP4BoxSTSS>& InStssBox, uint32 InMaxSamples)
+			{
+				StssBox = InStssBox;
+				NumTotalSamples = InMaxSamples;
+				CurrentSampleNum = 0;
+				CurrentEntryIndex = 0;
+				bCurrentIsSyncSample = !StssBox.IsValid() || (StssBox.IsValid() && StssBox->GetEntries().Num() && StssBox->GetEntries()[0] == 1);
+			}
+
+			void SetToSampleNumber(uint32 InSampleNum)
+			{
+				InSampleNum = InSampleNum < NumTotalSamples ? InSampleNum : NumTotalSamples-1;
+				if (StssBox.IsValid())
+				{
+					const TArray<uint32>& Entries(StssBox->GetEntries());
+					const uint32 SampleNumPlus1 = InSampleNum + 1;
+					for(CurrentEntryIndex=0; CurrentEntryIndex<(uint32)Entries.Num(); ++CurrentEntryIndex)
+					{
+						if (Entries[CurrentEntryIndex] >= SampleNumPlus1)
+						{
+							break;
+						}
+					}
+					bCurrentIsSyncSample = CurrentEntryIndex < (uint32)Entries.Num() ? Entries[CurrentEntryIndex] == SampleNumPlus1 : false;
+				}
+				CurrentSampleNum = InSampleNum;
+			}
+
+			bool Next()
+			{
+				if (IsLast())
+				{
+					return false;
+				}
+				++CurrentSampleNum;
+				if (StssBox.IsValid())
+				{
+					const TArray<uint32>& Entries(StssBox->GetEntries());
+					const uint32 SampleNumPlus1 = CurrentSampleNum + 1;
+					while(CurrentEntryIndex < (uint32)Entries.Num() && Entries[CurrentEntryIndex] < SampleNumPlus1)
+					{
+						++CurrentEntryIndex;
+					}
+					bCurrentIsSyncSample = CurrentEntryIndex < (uint32)Entries.Num() ? Entries[CurrentEntryIndex] == SampleNumPlus1 : false;
+				}
+				return true;
+			}
+
+			bool Prev()
+			{
+				if (IsFirst())
+				{
+					return false;
+				}
+				--CurrentSampleNum;
+				if (StssBox.IsValid())
+				{
+					const TArray<uint32>& Entries(StssBox->GetEntries());
+					const uint32 SampleNumPlus1 = CurrentSampleNum + 1;
+					if (Entries.Num())
+					{
+						CurrentEntryIndex = CurrentEntryIndex < (uint32)Entries.Num() ? CurrentEntryIndex : Entries.Num()-1;
+						while(CurrentEntryIndex > 0 && Entries[CurrentEntryIndex] > SampleNumPlus1)
+						{
+							--CurrentEntryIndex;
+						}
+					}
+					bCurrentIsSyncSample = CurrentEntryIndex < (uint32)Entries.Num() ? Entries[CurrentEntryIndex] == SampleNumPlus1 : false;
+				}
+				return true;
+			}
+
+			bool IsSyncSample() const
+			{
+				return bCurrentIsSyncSample;
+			}
+
+			uint32 GetNumTotalSamples() const
+			{ return NumTotalSamples; }
+			uint32 GetCurrentSampleNum() const
+			{ return CurrentSampleNum; }
+			bool IsFirst() const
+			{ return CurrentSampleNum == 0; }
+			bool IsLast() const
+			{ return CurrentSampleNum + 1 >= NumTotalSamples; }
+
+		private:
+			TSharedPtr<FMP4BoxSTSS> StssBox;
+			uint32 NumTotalSamples = 0;
+			uint32 CurrentSampleNum = 0;
+			uint32 CurrentEntryIndex = 0;
+			bool bCurrentIsSyncSample = true;
+		};
+
+		/****************************************************************************************************************************************************/
+
+		struct FSBGPBoxIterator
+		{
+			void SetBox(const TSharedPtr<FMP4BoxSBGP>& InSbgpBox, uint32 InDefaultSampleGroupDescriptionIndex, uint32 InMaxSamples)
+			{
+				NumTotalSamples = InMaxSamples;
+				Current_sample_count = InMaxSamples;
+				Current_group_description_index = 0;
+				CurrentEntryIndex = 0;
+				CurrentSampleNum = 0;
+				LocalSampleInGroup = 0;
+				SbgpBox = InSbgpBox;
+				if (SbgpBox.IsValid())
+				{
+					NumDescribedSamples = SbgpBox->GetNumTotalSamples();
+					DefaultSampleGroupDescriptionIndex = InDefaultSampleGroupDescriptionIndex;
+					check(SbgpBox->GetEntries().Num());
+					if (SbgpBox->GetEntries().Num())
+					{
+						const TArray<FMP4BoxSBGP::FEntry>& Entries(SbgpBox->GetEntries());
+						Current_sample_count = Entries[0].sample_count;
+						Current_group_description_index = Entries[0].group_description_index;
+					}
+				}
+			}
+
+			void SetToSampleNumber(uint32 InSampleNum)
+			{
+				InSampleNum = InSampleNum < NumTotalSamples ? InSampleNum : NumTotalSamples-1;
+
+				if (!SbgpBox.IsValid())
+				{
+					CurrentSampleNum = InSampleNum;
+					LocalSampleInGroup = InSampleNum;
+					return;
+				}
+				// It is permitted to have fewer sample entries here than there are samples elsewhere.
+				// In that case the default value is to be used.
+				if (InSampleNum < NumDescribedSamples)
+				{
+					const TArray<FMP4BoxSBGP::FEntry>& Entries(SbgpBox->GetEntries());
+					CurrentEntryIndex = 0;
+					Current_sample_count = Entries[0].sample_count;
+					Current_group_description_index = Entries[0].group_description_index;
+					uint32 n=0;
+					while(1)
+					{
+						n += Current_sample_count;
+						if (n > InSampleNum)
+						{
+							LocalSampleInGroup = Current_sample_count - (n - InSampleNum);
+							break;
+						}
+						else
+						{
+							check(CurrentEntryIndex < (uint32)Entries.Num());
+							++CurrentEntryIndex;
+							Current_sample_count = Entries[CurrentEntryIndex].sample_count;
+							Current_group_description_index = Entries[CurrentEntryIndex].group_description_index;
+						}
+					}
+				}
+				else
+				{
+					Current_sample_count = NumTotalSamples - NumDescribedSamples;
+					Current_group_description_index = DefaultSampleGroupDescriptionIndex;
+					LocalSampleInGroup = InSampleNum - NumDescribedSamples;
+					CurrentEntryIndex = SbgpBox->GetEntries().Num();
+				}
+				CurrentSampleNum = InSampleNum;
+			}
+
+			bool Next()
+			{
+				if (IsLast())
+				{
+					return false;
+				}
+
+				++CurrentSampleNum;
+				++LocalSampleInGroup;
+				if (!SbgpBox.IsValid() || CurrentSampleNum > NumDescribedSamples)
+				{
+					return true;
+				}
+				if (CurrentSampleNum == NumDescribedSamples)
+				{
+					Current_sample_count = NumTotalSamples - NumDescribedSamples;
+					Current_group_description_index = DefaultSampleGroupDescriptionIndex;
+					LocalSampleInGroup = CurrentSampleNum - NumDescribedSamples;
+					CurrentEntryIndex = SbgpBox->GetEntries().Num();
+					return true;
+				}
+				// Reached end of the current group run?
+				if (LocalSampleInGroup == Current_sample_count)
+				{
+					LocalSampleInGroup = 0;
+					const TArray<FMP4BoxSBGP::FEntry>& Entries(SbgpBox->GetEntries());
+					check(CurrentEntryIndex < (uint32)Entries.Num());
+					++CurrentEntryIndex;
+					Current_sample_count = Entries[CurrentEntryIndex].sample_count;
+					Current_group_description_index = Entries[CurrentEntryIndex].group_description_index;
+				}
+				return true;
+			}
+
+			bool Prev()
+			{
+				if (IsFirst())
+				{
+					return false;
+				}
+
+				--CurrentSampleNum;
+				if (!SbgpBox.IsValid() || CurrentSampleNum >= NumDescribedSamples)
+				{
+					--LocalSampleInGroup;
+					return true;
+				}
+
+				// At the start of the current chunk run?
+				if (LocalSampleInGroup > 0)
+				{
+					--LocalSampleInGroup;
+				}
+				else
+				{
+					const TArray<FMP4BoxSBGP::FEntry>& Entries(SbgpBox->GetEntries());
+					// Go back an entry if there is one.
+					if (CurrentEntryIndex)
+					{
+						--CurrentEntryIndex;
+						Current_sample_count = Entries[CurrentEntryIndex].sample_count;
+						Current_group_description_index = Entries[CurrentEntryIndex].group_description_index;
+					}
+					LocalSampleInGroup = Current_sample_count - 1;
+				}
+				return true;
+			}
+
+			uint32 GetCurrentGroupDescriptionIndex() const
+			{ return Current_group_description_index; }
+			uint32 GetNumSamplesInCurrentGroup() const
+			{ return Current_sample_count; }
+			uint32 GetSampleIndexInCurrentGroup() const
+			{ return LocalSampleInGroup; }
+			uint32 GetCurrentSampleNum() const
+			{ return CurrentSampleNum; }
+			bool IsFirst() const
+			{ return CurrentSampleNum == 0; }
+			bool IsLast() const
+			{ return CurrentSampleNum + 1 >= NumTotalSamples; }
+
+		private:
+			TSharedPtr<FMP4BoxSBGP> SbgpBox;
+			uint32 Current_sample_count = 0;
+			uint32 Current_group_description_index = 0;
+			uint32 DefaultSampleGroupDescriptionIndex = 0;
+
+			uint32 NumTotalSamples = 0;
+			uint32 NumDescribedSamples = 0;
+			uint32 CurrentSampleNum = 0;
+			uint32 LocalSampleInGroup = 0;
+			uint32 CurrentEntryIndex = 0;
+		};
+
+
+	} // namespace UtilitiesMP4
+
+} // namespace Electra
diff --git a/Engine/Plugins/Media/ElectraUtil/Source/ElectraBase/Public/Utilities/MP4Boxes/MP4BoxMetadata.h b/Engine/Plugins/Media/ElectraUtil/Source/ElectraBase/Public/Utilities/MP4Boxes/MP4BoxMetadata.h
new file mode 100644
--- /dev/null
+++ b/Engine/Plugins/Media/ElectraUtil/Source/ElectraBase/Public/Utilities/MP4Boxes/MP4BoxMetadata.h
@@ -0,0 +1,23 @@
+// Copyright Epic Games, Inc. All Rights Reserved.
+
+#pragma once
+
+#include "Utilities/MP4Boxes/MP4Boxes.h"
+
+namespace Electra
+{
+	namespace UtilitiesMP4
+	{
+
+		/****************************************************************************************************************************************************/
+
+		struct FMP4TrackMetadataCommon
+		{
+			FString Language;
+			FString Name;
+			FString HandlerName;
+		};
+
+	} // namespace UtilitiesMP4
+
+} // namespace Electra
diff --git a/Engine/Plugins/Media/ElectraUtil/Source/ElectraBase/Public/Utilities/MP4Boxes/MP4Boxes.h b/Engine/Plugins/Media/ElectraUtil/Source/ElectraBase/Public/Utilities/MP4Boxes/MP4Boxes.h
new file mode 100644
--- /dev/null
+++ b/Engine/Plugins/Media/ElectraUtil/Source/ElectraBase/Public/Utilities/MP4Boxes/MP4Boxes.h
@@ -0,0 +1,1846 @@
+// Copyright Epic Games, Inc. All Rights Reserved.
+
+#pragma once
+
+#include "Utilities/MP4Boxes/MP4BoxBase.h"
+#include "Containers/ChunkedArray.h"
+#include "PlayerTime.h"
+
+namespace Electra
+{
+	namespace UtilitiesMP4
+	{
+
+		/**
+		 * `free` and `skip` box
+		 * ISO/IEC 14496-12:2022 - 8.1.2 Free space box
+		 */
+		class FMP4BoxFREE : public FMP4BoxBase
+		{
+		public:
+			static TSharedPtr<FMP4BoxBase> Create(const TWeakPtr<FMP4BoxBase>& InParent, const FMP4BoxInfo& InBoxInfo)
+			{ return MakeShareable<FMP4BoxFREE>(new FMP4BoxFREE(InParent, InBoxInfo)); }
+			virtual ~FMP4BoxFREE() = default;
+		protected:
+			FMP4BoxFREE(const TWeakPtr<FMP4BoxBase>& InParent, const FMP4BoxInfo& InBoxInfo)
+				: FMP4BoxBase(InParent, InBoxInfo)
+			{ }
+		};
+
+
+		/**
+		 * `ftyp` box, `styp` box
+		 * ISO/IEC 14496-12:2022 - 4.3 File-type box & 8.16.2 Segment type box
+		 */
+		class FMP4BoxFTYP : public FMP4BoxBase
+		{
+		public:
+			static TSharedPtr<FMP4BoxBase> Create(const TWeakPtr<FMP4BoxBase>& InParent, const FMP4BoxInfo& InBoxInfo)
+			{ return MakeShareable<FMP4BoxFTYP>(new FMP4BoxFTYP(InParent, InBoxInfo)); }
+			virtual ~FMP4BoxFTYP() = default;
+		protected:
+			FMP4BoxFTYP(const TWeakPtr<FMP4BoxBase>& InParent, const FMP4BoxInfo& InBoxInfo)
+				: FMP4BoxBase(InParent, InBoxInfo)
+			{ }
+		};
+
+
+		/**
+		 * `moov` box
+		 * ISO/IEC 14496-12:2022 - 8.2.1 Movie box
+		 */
+		class FMP4BoxMOOV : public FMP4BoxBase
+		{
+		public:
+			static TSharedPtr<FMP4BoxBase> Create(const TWeakPtr<FMP4BoxBase>& InParent, const FMP4BoxInfo& InBoxInfo)
+			{ return MakeShareable<FMP4BoxMOOV>(new FMP4BoxMOOV(InParent, InBoxInfo)); }
+			virtual ~FMP4BoxMOOV() = default;
+			bool IsLeafBox() const override { return false; }
+		protected:
+			FMP4BoxMOOV(const TWeakPtr<FMP4BoxBase>& InParent, const FMP4BoxInfo& InBoxInfo)
+				: FMP4BoxBase(InParent, InBoxInfo)
+			{ }
+		};
+
+
+		/**
+		 * `mvhd` box
+		 * ISO/IEC 14496-12:2022 - 8.2.2 Movie header box
+		 */
+		class FMP4BoxMVHD : public FMP4BoxBase
+		{
+		public:
+			static TSharedPtr<FMP4BoxBase> Create(const TWeakPtr<FMP4BoxBase>& InParent, const FMP4BoxInfo& InBoxInfo)
+			{ return MakeShareable<FMP4BoxMVHD>(new FMP4BoxMVHD(InParent, InBoxInfo)); }
+			virtual ~FMP4BoxMVHD() = default;
+
+			ELECTRABASE_API FTimeFraction GetDuration();
+			ELECTRABASE_API uint32 GetTimescale();
+		protected:
+			FMP4BoxMVHD(const TWeakPtr<FMP4BoxBase>& InParent, const FMP4BoxInfo& InBoxInfo)
+				: FMP4BoxBase(InParent, InBoxInfo)
+			{ }
+			void ParseIfRequired();
+			struct FParsed;
+			TPimplPtr<FParsed> Parsed;
+		};
+
+
+		/**
+		 * `trak` box
+		 * ISO/IEC 14496-12:2022 - 8.3.1 Track box
+		 */
+		class FMP4BoxTRAK : public FMP4BoxBase
+		{
+		public:
+			static TSharedPtr<FMP4BoxBase> Create(const TWeakPtr<FMP4BoxBase>& InParent, const FMP4BoxInfo& InBoxInfo)
+			{ return MakeShareable<FMP4BoxTRAK>(new FMP4BoxTRAK(InParent, InBoxInfo)); }
+			virtual ~FMP4BoxTRAK() = default;
+			bool IsLeafBox() const override { return false; }
+		protected:
+			FMP4BoxTRAK(const TWeakPtr<FMP4BoxBase>& InParent, const FMP4BoxInfo& InBoxInfo)
+				: FMP4BoxBase(InParent, InBoxInfo)
+			{ }
+		};
+
+
+		/**
+		 * `tkhd` box
+		 * ISO/IEC 14496-12:2022 - 8.3.2 Track header box
+		 */
+		class FMP4BoxTKHD : public FMP4BoxBase
+		{
+		public:
+			static TSharedPtr<FMP4BoxBase> Create(const TWeakPtr<FMP4BoxBase>& InParent, const FMP4BoxInfo& InBoxInfo)
+			{ return MakeShareable<FMP4BoxTKHD>(new FMP4BoxTKHD(InParent, InBoxInfo)); }
+			virtual ~FMP4BoxTKHD() = default;
+
+			ELECTRABASE_API uint32 GetFlags();
+
+			bool IsEnabled()
+			{ return !!(GetFlags() & 1); }
+
+			bool IsInMovie()
+			{ return !!(GetFlags() & 2); }
+
+			bool IsInPreview()
+			{ return !!(GetFlags() & 4); }
+
+			bool IsTrackSizeAspectRatio()
+			{ return !!(GetFlags() & 8); }
+
+			/**
+			 * Returns the duration of this track, measured in the timescale of the `mvhd` box.
+			 * This is just the value as stored in the box. You need to apply the timescale yourself.
+			 */
+			ELECTRABASE_API int64 GetDuration();
+
+			ELECTRABASE_API uint32 GetTrackID();
+
+			ELECTRABASE_API uint16 GetWidth();
+			ELECTRABASE_API uint16 GetHeight();
+		protected:
+			FMP4BoxTKHD(const TWeakPtr<FMP4BoxBase>& InParent, const FMP4BoxInfo& InBoxInfo)
+				: FMP4BoxBase(InParent, InBoxInfo)
+			{ }
+			void ParseIfRequired();
+			struct FParsed;
+			TPimplPtr<FParsed> Parsed;
+		};
+
+
+		/**
+		 * `tref` box
+		 * ISO/IEC 14496-12:2022 - 8.3.3 Track reference box
+		 */
+		class FMP4BoxTREF : public FMP4BoxBase
+		{
+		public:
+			static TSharedPtr<FMP4BoxBase> Create(const TWeakPtr<FMP4BoxBase>& InParent, const FMP4BoxInfo& InBoxInfo)
+			{ return MakeShareable<FMP4BoxTREF>(new FMP4BoxTREF(InParent, InBoxInfo)); }
+			virtual ~FMP4BoxTREF() = default;
+			bool IsLeafBox() const override { return false; }
+			bool IsListOfEntries() const override { return true; }
+
+			struct FEntry
+			{
+				uint32 Type = 0;
+				TArray<uint32> TrackIDs;
+			};
+
+			ELECTRABASE_API const TArray<FEntry>& GetEntries();
+			ELECTRABASE_API TArray<FEntry> GetEntriesOfType(uint32 InType);
+		protected:
+			FMP4BoxTREF(const TWeakPtr<FMP4BoxBase>& InParent, const FMP4BoxInfo& InBoxInfo)
+				: FMP4BoxBase(InParent, InBoxInfo)
+			{ }
+			void ParseIfRequired();
+			struct FParsed;
+			TPimplPtr<FParsed> Parsed;
+		};
+
+
+		/**
+		 * `edts` box
+		 * ISO/IEC 14496-12:2022 - 8.6.5 Edit box
+		 */
+		class FMP4BoxEDTS : public FMP4BoxBase
+		{
+		public:
+			static TSharedPtr<FMP4BoxBase> Create(const TWeakPtr<FMP4BoxBase>& InParent, const FMP4BoxInfo& InBoxInfo)
+			{ return MakeShareable<FMP4BoxEDTS>(new FMP4BoxEDTS(InParent, InBoxInfo)); }
+			virtual ~FMP4BoxEDTS() = default;
+			bool IsLeafBox() const override { return false; }
+		protected:
+			FMP4BoxEDTS(const TWeakPtr<FMP4BoxBase>& InParent, const FMP4BoxInfo& InBoxInfo)
+				: FMP4BoxBase(InParent, InBoxInfo)
+			{ }
+		};
+
+
+		/**
+		 * `elst` box
+		 * ISO/IEC 14496-12:2022 - 8.6.6 Edit list box
+		 */
+		class FMP4BoxELST : public FMP4BoxBase
+		{
+		public:
+			static TSharedPtr<FMP4BoxBase> Create(const TWeakPtr<FMP4BoxBase>& InParent, const FMP4BoxInfo& InBoxInfo)
+			{ return MakeShareable<FMP4BoxELST>(new FMP4BoxELST(InParent, InBoxInfo)); }
+			virtual ~FMP4BoxELST() = default;
+
+			struct FEntry
+			{
+				// Specified in units of timescale of the `mvhd` box.
+				int64 edit_duration = 0;				// A uint64 in the standard, but FTimeFraction needs an int64 for the numerator.
+				// Specified in units of timescale of the `mdhd` box of this track.
+				int64 media_time = 0;
+				int16 media_rate_integer = 0;
+				int16 media_rate_fraction = 0;
+			};
+
+			ELECTRABASE_API const TArray<FEntry>& GetEntries();
+			ELECTRABASE_API bool RepeatEdits();
+
+		protected:
+			FMP4BoxELST(const TWeakPtr<FMP4BoxBase>& InParent, const FMP4BoxInfo& InBoxInfo)
+				: FMP4BoxBase(InParent, InBoxInfo)
+			{ }
+			void ParseIfRequired();
+			struct FParsed;
+			TPimplPtr<FParsed> Parsed;
+		};
+
+
+		/**
+		 * `mdia` box
+		 * ISO/IEC 14496-12:2022 - 8.4.1 Media box
+		 */
+		class FMP4BoxMDIA : public FMP4BoxBase
+		{
+		public:
+			static TSharedPtr<FMP4BoxBase> Create(const TWeakPtr<FMP4BoxBase>& InParent, const FMP4BoxInfo& InBoxInfo)
+			{ return MakeShareable<FMP4BoxMDIA>(new FMP4BoxMDIA(InParent, InBoxInfo)); }
+			virtual ~FMP4BoxMDIA() = default;
+			bool IsLeafBox() const override { return false; }
+		protected:
+			FMP4BoxMDIA(const TWeakPtr<FMP4BoxBase>& InParent, const FMP4BoxInfo& InBoxInfo)
+				: FMP4BoxBase(InParent, InBoxInfo)
+			{ }
+		};
+
+
+		/**
+		 * `mdhd` box
+		 * ISO/IEC 14496-12:2022 - 8.4.2 Media header box
+		 */
+		class FMP4BoxMDHD : public FMP4BoxBase
+		{
+		public:
+			static TSharedPtr<FMP4BoxBase> Create(const TWeakPtr<FMP4BoxBase>& InParent, const FMP4BoxInfo& InBoxInfo)
+			{ return MakeShareable<FMP4BoxMDHD>(new FMP4BoxMDHD(InParent, InBoxInfo)); }
+			virtual ~FMP4BoxMDHD() = default;
+
+			ELECTRABASE_API FTimeFraction GetDuration();
+			ELECTRABASE_API uint32 GetTimescale();
+			ELECTRABASE_API const FString& GetLanguage();
+
+		protected:
+			FMP4BoxMDHD(const TWeakPtr<FMP4BoxBase>& InParent, const FMP4BoxInfo& InBoxInfo)
+				: FMP4BoxBase(InParent, InBoxInfo)
+			{ }
+			void ParseIfRequired();
+			struct FParsed;
+			TPimplPtr<FParsed> Parsed;
+		};
+
+
+		/**
+		 * `hdlr` box
+		 * ISO/IEC 14496-12:2022 - 8.4.3 Handler reference box
+		 */
+		class FMP4BoxHDLR : public FMP4BoxBase
+		{
+		public:
+			static TSharedPtr<FMP4BoxBase> Create(const TWeakPtr<FMP4BoxBase>& InParent, const FMP4BoxInfo& InBoxInfo)
+			{ return MakeShareable<FMP4BoxHDLR>(new FMP4BoxHDLR(InParent, InBoxInfo)); }
+			virtual ~FMP4BoxHDLR() = default;
+
+			ELECTRABASE_API uint32 GetHandlerType();
+			ELECTRABASE_API FString GetHandlerName();
+		protected:
+			FMP4BoxHDLR(const TWeakPtr<FMP4BoxBase>& InParent, const FMP4BoxInfo& InBoxInfo)
+				: FMP4BoxBase(InParent, InBoxInfo)
+			{ }
+			void ParseIfRequired();
+			struct FParsed;
+			TPimplPtr<FParsed> Parsed;
+		};
+
+
+		/**
+		 * `minf` box
+		 * ISO/IEC 14496-12:2022 - 8.4.4 Media information box
+		 */
+		class FMP4BoxMINF : public FMP4BoxBase
+		{
+		public:
+			static TSharedPtr<FMP4BoxBase> Create(const TWeakPtr<FMP4BoxBase>& InParent, const FMP4BoxInfo& InBoxInfo)
+			{ return MakeShareable<FMP4BoxMINF>(new FMP4BoxMINF(InParent, InBoxInfo)); }
+			virtual ~FMP4BoxMINF() = default;
+			bool IsLeafBox() const override { return false; }
+		protected:
+			FMP4BoxMINF(const TWeakPtr<FMP4BoxBase>& InParent, const FMP4BoxInfo& InBoxInfo)
+				: FMP4BoxBase(InParent, InBoxInfo)
+			{ }
+		};
+
+
+		/**
+		 * `elng` box
+		 * ISO/IEC 14496-12:2022 - 8.4.6 Extended language tag
+		 */
+		class FMP4BoxELNG : public FMP4BoxBase
+		{
+		public:
+			static TSharedPtr<FMP4BoxBase> Create(const TWeakPtr<FMP4BoxBase>& InParent, const FMP4BoxInfo& InBoxInfo)
+			{ return MakeShareable<FMP4BoxELNG>(new FMP4BoxELNG(InParent, InBoxInfo)); }
+			virtual ~FMP4BoxELNG() = default;
+
+			ELECTRABASE_API FString GetLanguage();
+		protected:
+			FMP4BoxELNG(const TWeakPtr<FMP4BoxBase>& InParent, const FMP4BoxInfo& InBoxInfo)
+				: FMP4BoxBase(InParent, InBoxInfo)
+			{ }
+			void ParseIfRequired();
+			struct FParsed;
+			TPimplPtr<FParsed> Parsed;
+		};
+
+
+		/**
+		 * `dinf` box
+		 * ISO/IEC 14496-12:2022 - 8.7.1 Data information box
+		 */
+		class FMP4BoxDINF : public FMP4BoxBase
+		{
+		public:
+			static TSharedPtr<FMP4BoxBase> Create(const TWeakPtr<FMP4BoxBase>& InParent, const FMP4BoxInfo& InBoxInfo)
+			{ return MakeShareable<FMP4BoxDINF>(new FMP4BoxDINF(InParent, InBoxInfo)); }
+			virtual ~FMP4BoxDINF() = default;
+			bool IsLeafBox() const override { return false; }
+		protected:
+			FMP4BoxDINF(const TWeakPtr<FMP4BoxBase>& InParent, const FMP4BoxInfo& InBoxInfo)
+				: FMP4BoxBase(InParent, InBoxInfo)
+			{ }
+		};
+
+
+		/**
+		 * `dref` box
+		 * ISO/IEC 14496-12:2022 - 8.7.2 Data reference box
+		 */
+		class FMP4BoxDREF : public FMP4BoxBase
+		{
+		public:
+			static TSharedPtr<FMP4BoxBase> Create(const TWeakPtr<FMP4BoxBase>& InParent, const FMP4BoxInfo& InBoxInfo)
+			{ return MakeShareable<FMP4BoxDREF>(new FMP4BoxDREF(InParent, InBoxInfo)); }
+			virtual ~FMP4BoxDREF() = default;
+			bool IsLeafBox() const override { return false; }
+			bool IsListOfEntries() const override { return true; }
+		protected:
+			FMP4BoxDREF(const TWeakPtr<FMP4BoxBase>& InParent, const FMP4BoxInfo& InBoxInfo)
+				: FMP4BoxBase(InParent, InBoxInfo)
+			{ }
+		};
+
+
+		/**
+		 * `stbl` box
+		 * ISO/IEC 14496-12:2022 - 8.5.1 Sample table box
+		 */
+		class FMP4BoxSTBL : public FMP4BoxBase
+		{
+		public:
+			static TSharedPtr<FMP4BoxBase> Create(const TWeakPtr<FMP4BoxBase>& InParent, const FMP4BoxInfo& InBoxInfo)
+			{ return MakeShareable<FMP4BoxSTBL>(new FMP4BoxSTBL(InParent, InBoxInfo)); }
+			virtual ~FMP4BoxSTBL() = default;
+			bool IsLeafBox() const override { return false; }
+		protected:
+			FMP4BoxSTBL(const TWeakPtr<FMP4BoxBase>& InParent, const FMP4BoxInfo& InBoxInfo)
+				: FMP4BoxBase(InParent, InBoxInfo)
+			{ }
+		};
+
+
+		/**
+		 * Sample entry base class
+		 * ISO/IEC 14496-12:2022 - 8.5.2 Sample description box
+		 */
+		class FMP4BoxSampleEntry : public FMP4BoxBase
+		{
+		public:
+			virtual ~FMP4BoxSampleEntry() = default;
+			bool IsSampleDescription() const override
+			{ return true; }
+
+			enum class ESampleType
+			{
+				Video,
+				Audio,
+				Subtitles,
+				QTFFTimecode,
+				TimedMetadata,
+				Unsupported
+			};
+			virtual ESampleType GetSampleType()
+			{ return ESampleType::Unsupported; }
+			virtual uint16 GetDataReferenceIndex() = 0;
+
+		protected:
+			FMP4BoxSampleEntry(const TWeakPtr<FMP4BoxBase>& InParent, const FMP4BoxInfo& InBoxInfo)
+				: FMP4BoxBase(InParent, InBoxInfo)
+			{ }
+		};
+
+
+		/**
+		 * `btrt` box
+		 * ISO/IEC 14496-12:2022 - 8.5.2 Sample description box (8.5.2.2 Bitrate box)
+		 */
+		class FMP4BoxBTRT : public FMP4BoxBase
+		{
+		public:
+			static TSharedPtr<FMP4BoxBase> Create(const TWeakPtr<FMP4BoxBase>& InParent, const FMP4BoxInfo& InBoxInfo)
+			{ return MakeShareable<FMP4BoxBTRT>(new FMP4BoxBTRT(InParent, InBoxInfo)); }
+			virtual ~FMP4BoxBTRT() = default;
+			ELECTRABASE_API uint32 GetBufferSizeDB();
+			ELECTRABASE_API uint32 GetMaxBitrate();
+			ELECTRABASE_API uint32 GetAverageBitrate();
+		protected:
+			FMP4BoxBTRT(const TWeakPtr<FMP4BoxBase>& InParent, const FMP4BoxInfo& InBoxInfo)
+				: FMP4BoxBase(InParent, InBoxInfo)
+			{ }
+			void ParseIfRequired();
+			struct FParsed;
+			TPimplPtr<FParsed> Parsed;
+		};
+
+
+		/**
+		 * Visual sample entry
+		 * ISO/IEC 14496-12:2022 - 12.1.3 Sample entry
+		 */
+		class FMP4BoxVisualSampleEntry : public FMP4BoxSampleEntry
+		{
+		public:
+			static TSharedPtr<FMP4BoxBase> Create(const TWeakPtr<FMP4BoxBase>& InParent, const FMP4BoxInfo& InBoxInfo)
+			{ return MakeShareable<FMP4BoxVisualSampleEntry>(new FMP4BoxVisualSampleEntry(InParent, InBoxInfo)); }
+			virtual ~FMP4BoxVisualSampleEntry() = default;
+			ESampleType GetSampleType() override
+			{ return ESampleType::Video; }
+			uint16 GetDataReferenceIndex() override;
+			ELECTRABASE_API uint16 GetWidth();
+			ELECTRABASE_API uint16 GetHeight();
+			ELECTRABASE_API uint16 GetFrameCount();
+			ELECTRABASE_API uint16 GetDepth();
+		protected:
+			FMP4BoxVisualSampleEntry(const TWeakPtr<FMP4BoxBase>& InParent, const FMP4BoxInfo& InBoxInfo)
+				: FMP4BoxSampleEntry(InParent, InBoxInfo)
+			{ }
+			void ParseIfRequired();
+			struct FParsed;
+			TPimplPtr<FParsed> Parsed;
+		};
+
+
+		/**
+		 * Audio sample entry
+		 * ISO/IEC 14496-12:2022 - 12.2.3 Sample entry
+		 */
+		class FMP4BoxAudioSampleEntry : public FMP4BoxSampleEntry
+		{
+		public:
+			static TSharedPtr<FMP4BoxBase> Create(const TWeakPtr<FMP4BoxBase>& InParent, const FMP4BoxInfo& InBoxInfo)
+			{ return MakeShareable<FMP4BoxAudioSampleEntry>(new FMP4BoxAudioSampleEntry(InParent, InBoxInfo)); }
+			virtual ~FMP4BoxAudioSampleEntry() = default;
+			ESampleType GetSampleType() override
+			{ return ESampleType::Audio; }
+			uint16 GetDataReferenceIndex() override;
+			ELECTRABASE_API uint32 GetSampleRate();
+			ELECTRABASE_API int32 GetChannelCount();
+			ELECTRABASE_API int32 GetSampleSize();
+			ELECTRABASE_API bool HaveFormatSpecificFlags();
+			ELECTRABASE_API int32 GetFormatSpecificFlags();
+			ELECTRABASE_API uint32 GetConstBytesPerAudioPacket();
+			ELECTRABASE_API uint32 GetConstLPCMFramesPerAudioPacket();
+		protected:
+			FMP4BoxAudioSampleEntry(const TWeakPtr<FMP4BoxBase>& InParent, const FMP4BoxInfo& InBoxInfo)
+				: FMP4BoxSampleEntry(InParent, InBoxInfo)
+			{ }
+			void ParseIfRequired();
+			struct FParsed;
+			TPimplPtr<FParsed> Parsed;
+		};
+
+
+		/**
+		 * Timecode sample entry
+		 * Apple QuickTime: https://developer.apple.com/documentation/quicktime-file-format/timecode_sample_description
+		 */
+		class FMP4BoxQTFFTimecodeSampleEntry : public FMP4BoxSampleEntry
+		{
+		public:
+			static TSharedPtr<FMP4BoxBase> Create(const TWeakPtr<FMP4BoxBase>& InParent, const FMP4BoxInfo& InBoxInfo)
+			{ return MakeShareable<FMP4BoxQTFFTimecodeSampleEntry>(new FMP4BoxQTFFTimecodeSampleEntry(InParent, InBoxInfo)); }
+			virtual ~FMP4BoxQTFFTimecodeSampleEntry() = default;
+			ESampleType GetSampleType() override
+			{ return ESampleType::QTFFTimecode; }
+			uint16 GetDataReferenceIndex() override;
+
+			// See: https://developer.apple.com/documentation/quicktime-file-format/timecode_sample_description/flags
+			enum EFlags : uint32
+			{
+				DropFrame = 0x0001,					// Indicates whether the timecode is drop frame. Set it to 1 if the timecode is drop frame.
+				Max24Hour = 0x0002,					// Indicates whether the timecode wraps after 24 hours. Set it to 1 if the timecode wraps.
+				AllowNegativeTimes = 0x0004,		// Indicates whether negative time values are allowed. Set it to 1 if the timecode supports negative values.
+				Counter = 0x0008					// Indicates whether the time value corresponds to a tape counter value. Set it to 1 if the timecode values are tape counter values.
+			};
+			ELECTRABASE_API uint32 GetFlags();
+			ELECTRABASE_API uint32 GetTimescale();
+			ELECTRABASE_API uint32 GetFrameDuration();
+			ELECTRABASE_API uint32 GetNumberOfFrames();
+		protected:
+			FMP4BoxQTFFTimecodeSampleEntry(const TWeakPtr<FMP4BoxBase>& InParent, const FMP4BoxInfo& InBoxInfo)
+				: FMP4BoxSampleEntry(InParent, InBoxInfo)
+			{ }
+			void ParseIfRequired();
+			struct FParsed;
+			TPimplPtr<FParsed> Parsed;
+		};
+
+
+		/**
+		 * 3GPP / TX3G Text sample entry (ETSI TS 126 245 V11.0.0 - 5.16 Sample Description Format)
+		 */
+		class FMP4BoxTX3GSampleEntry : public FMP4BoxSampleEntry
+		{
+		public:
+			static TSharedPtr<FMP4BoxBase> Create(const TWeakPtr<FMP4BoxBase>& InParent, const FMP4BoxInfo& InBoxInfo)
+			{ return MakeShareable<FMP4BoxTX3GSampleEntry>(new FMP4BoxTX3GSampleEntry(InParent, InBoxInfo)); }
+			virtual ~FMP4BoxTX3GSampleEntry() = default;
+			ESampleType GetSampleType() override
+			{ return ESampleType::Subtitles; }
+			uint16 GetDataReferenceIndex() override;
+		protected:
+			FMP4BoxTX3GSampleEntry(const TWeakPtr<FMP4BoxBase>& InParent, const FMP4BoxInfo& InBoxInfo)
+				: FMP4BoxSampleEntry(InParent, InBoxInfo)
+			{ }
+			void ParseIfRequired();
+			struct FParsed;
+			TPimplPtr<FParsed> Parsed;
+		};
+
+
+		/**
+		 * `pasp` box
+		 * ISO/IEC 14496-12:2022 - 12.1.4 Pixel Aspect Ratio
+		 */
+		class FMP4BoxPASP : public FMP4BoxBase
+		{
+		public:
+			static TSharedPtr<FMP4BoxBase> Create(const TWeakPtr<FMP4BoxBase>& InParent, const FMP4BoxInfo& InBoxInfo)
+			{ return MakeShareable<FMP4BoxPASP>(new FMP4BoxPASP(InParent, InBoxInfo)); }
+			virtual ~FMP4BoxPASP() = default;
+		protected:
+			FMP4BoxPASP(const TWeakPtr<FMP4BoxBase>& InParent, const FMP4BoxInfo& InBoxInfo)
+				: FMP4BoxBase(InParent, InBoxInfo)
+			{ }
+		};
+
+
+		/**
+		 * `colr` box
+		 * ISO/IEC 14496-12:2022 - 12.1.5 Colour information
+		 */
+		class FMP4BoxCOLR : public FMP4BoxBase
+		{
+		public:
+			static TSharedPtr<FMP4BoxBase> Create(const TWeakPtr<FMP4BoxBase>& InParent, const FMP4BoxInfo& InBoxInfo)
+			{ return MakeShareable<FMP4BoxCOLR>(new FMP4BoxCOLR(InParent, InBoxInfo)); }
+			virtual ~FMP4BoxCOLR() = default;
+			enum class EColorType
+			{
+				nclx,
+				rICC,
+				prof,
+				Unsupported
+			};
+			struct FColorNCLX
+			{
+				uint16 colour_primaries = 0;
+				uint16 transfer_characteristics = 0;
+				uint16 matrix_coefficients;
+				uint8 full_range_flag = 0;
+			};
+			ELECTRABASE_API EColorType GetColorType();
+			ELECTRABASE_API const FColorNCLX& GetColorNCLX();
+		protected:
+			FMP4BoxCOLR(const TWeakPtr<FMP4BoxBase>& InParent, const FMP4BoxInfo& InBoxInfo)
+				: FMP4BoxBase(InParent, InBoxInfo)
+			{ }
+			void ParseIfRequired();
+			struct FParsed;
+			TPimplPtr<FParsed> Parsed;
+		};
+
+
+		/**
+		 * `clli` box
+		 * ISO/IEC 14496-12:2022 - 12.1.6 Content light level
+		 */
+		class FMP4BoxCLLI : public FMP4BoxBase
+		{
+		public:
+			static TSharedPtr<FMP4BoxBase> Create(const TWeakPtr<FMP4BoxBase>& InParent, const FMP4BoxInfo& InBoxInfo)
+			{ return MakeShareable<FMP4BoxCLLI>(new FMP4BoxCLLI(InParent, InBoxInfo)); }
+			virtual ~FMP4BoxCLLI() = default;
+			//ELECTRABASE_API xxxx GetXXXX();
+		protected:
+			FMP4BoxCLLI(const TWeakPtr<FMP4BoxBase>& InParent, const FMP4BoxInfo& InBoxInfo)
+				: FMP4BoxBase(InParent, InBoxInfo)
+			{ }
+			void ParseIfRequired();
+			struct FParsed;
+			TPimplPtr<FParsed> Parsed;
+		};
+
+
+		/**
+		 * `mdcv` box
+		 * ISO/IEC 14496-12:2022 - 12.1.7 Mastering display colour volume
+		 */
+		class FMP4BoxMDCV : public FMP4BoxBase
+		{
+		public:
+			static TSharedPtr<FMP4BoxBase> Create(const TWeakPtr<FMP4BoxBase>& InParent, const FMP4BoxInfo& InBoxInfo)
+			{ return MakeShareable<FMP4BoxMDCV>(new FMP4BoxMDCV(InParent, InBoxInfo)); }
+			virtual ~FMP4BoxMDCV() = default;
+			//ELECTRABASE_API xxxx GetXXXX();
+		protected:
+			FMP4BoxMDCV(const TWeakPtr<FMP4BoxBase>& InParent, const FMP4BoxInfo& InBoxInfo)
+				: FMP4BoxBase(InParent, InBoxInfo)
+			{ }
+			void ParseIfRequired();
+			struct FParsed;
+			TPimplPtr<FParsed> Parsed;
+		};
+
+
+		/**
+		 * `stsd` box
+		 * ISO/IEC 14496-12:2022 - 8.5.2 Sample description box
+		 */
+		class FMP4BoxSTSD : public FMP4BoxBase
+		{
+		public:
+			static TSharedPtr<FMP4BoxBase> Create(const TWeakPtr<FMP4BoxBase>& InParent, const FMP4BoxInfo& InBoxInfo)
+			{ return MakeShareable<FMP4BoxSTSD>(new FMP4BoxSTSD(InParent, InBoxInfo)); }
+			virtual ~FMP4BoxSTSD() = default;
+			bool IsLeafBox() const override { return false; }
+			bool IsListOfEntries() const override { return true; }
+
+			ELECTRABASE_API uint8 GetBoxVersion();
+			ELECTRABASE_API FMP4BoxSampleEntry::ESampleType GetSampleType();
+
+			template<typename T>
+			void GetSampleDescriptions(TArray<TSharedPtr<T>>& OutChildren)
+			{
+				// Get the sample type managed by this box. This may implicitly trigger parsing!
+				const FMP4BoxSampleEntry::ESampleType SampleType = GetSampleType();
+				for(auto& It : GetChildren())
+				{
+					if (It->IsSampleDescription() && StaticCastSharedPtr<FMP4BoxSampleEntry>(It)->GetSampleType() == SampleType)
+					{
+						uint16 DummyParseTrigger = StaticCastSharedPtr<FMP4BoxSampleEntry>(It)->GetDataReferenceIndex();
+						(void)DummyParseTrigger;
+
+						OutChildren.Emplace(StaticCastSharedPtr<T>(It));
+					}
+				}
+			}
+
+		protected:
+			FMP4BoxSTSD(const TWeakPtr<FMP4BoxBase>& InParent, const FMP4BoxInfo& InBoxInfo)
+				: FMP4BoxBase(InParent, InBoxInfo)
+			{ }
+			void ParseIfRequired();
+			struct FParsed;
+			TPimplPtr<FParsed> Parsed;
+		};
+
+
+		/**
+		 * `stts` box
+		 * ISO/IEC 14496-12:2022 - 8.6.1.2 Decoding time to sample box
+		 */
+		class FMP4BoxSTTS : public FMP4BoxBase
+		{
+		public:
+			static TSharedPtr<FMP4BoxBase> Create(const TWeakPtr<FMP4BoxBase>& InParent, const FMP4BoxInfo& InBoxInfo)
+			{ return MakeShareable<FMP4BoxSTTS>(new FMP4BoxSTTS(InParent, InBoxInfo)); }
+			virtual ~FMP4BoxSTTS() = default;
+			struct FEntry
+			{
+				uint32 sample_count = 0;
+				uint32 sample_delta = 0;
+			};
+			ELECTRABASE_API const TArray<FEntry>& GetEntries();
+			ELECTRABASE_API uint32 GetNumTotalSamples();
+			ELECTRABASE_API int64 GetTotalDuration();
+		protected:
+			FMP4BoxSTTS(const TWeakPtr<FMP4BoxBase>& InParent, const FMP4BoxInfo& InBoxInfo)
+				: FMP4BoxBase(InParent, InBoxInfo)
+			{ }
+			void ParseIfRequired();
+			struct FParsed;
+			TPimplPtr<FParsed> Parsed;
+		};
+
+
+		/**
+		 * `ctts` box
+		 * ISO/IEC 14496-12:2022 - 8.6.1.3 Composition time to sample box
+		 */
+		class FMP4BoxCTTS : public FMP4BoxBase
+		{
+		public:
+			static TSharedPtr<FMP4BoxBase> Create(const TWeakPtr<FMP4BoxBase>& InParent, const FMP4BoxInfo& InBoxInfo)
+			{ return MakeShareable<FMP4BoxCTTS>(new FMP4BoxCTTS(InParent, InBoxInfo)); }
+			virtual ~FMP4BoxCTTS() = default;
+			struct FEntry
+			{
+				// A 64 bit value used here to hold both signed values (version 1) and 32 bit unsigned values (version 0).
+				int64 sample_offset = 0;
+				uint32 sample_count = 0;
+			};
+			ELECTRABASE_API uint8 GetBoxVersion();
+			ELECTRABASE_API const TArray<FEntry>& GetEntries();
+			ELECTRABASE_API uint32 GetNumTotalSamples();
+		protected:
+			FMP4BoxCTTS(const TWeakPtr<FMP4BoxBase>& InParent, const FMP4BoxInfo& InBoxInfo)
+				: FMP4BoxBase(InParent, InBoxInfo)
+			{ }
+			void ParseIfRequired();
+			struct FParsed;
+			TPimplPtr<FParsed> Parsed;
+		};
+
+
+		/**
+		 * `cslg` box
+		 * ISO/IEC 14496-12:2022 - 8.6.1.4 Composition to decode box
+		 */
+		class FMP4BoxCSLG : public FMP4BoxBase
+		{
+		public:
+			static TSharedPtr<FMP4BoxBase> Create(const TWeakPtr<FMP4BoxBase>& InParent, const FMP4BoxInfo& InBoxInfo)
+			{ return MakeShareable<FMP4BoxCSLG>(new FMP4BoxCSLG(InParent, InBoxInfo)); }
+			virtual ~FMP4BoxCSLG() = default;
+		protected:
+			FMP4BoxCSLG(const TWeakPtr<FMP4BoxBase>& InParent, const FMP4BoxInfo& InBoxInfo)
+				: FMP4BoxBase(InParent, InBoxInfo)
+			{ }
+		};
+
+
+		/**
+		 * `stss` box
+		 * ISO/IEC 14496-12:2022 - 8.6.2 Sync sample box
+		 */
+		class FMP4BoxSTSS : public FMP4BoxBase
+		{
+		public:
+			static TSharedPtr<FMP4BoxBase> Create(const TWeakPtr<FMP4BoxBase>& InParent, const FMP4BoxInfo& InBoxInfo)
+			{ return MakeShareable<FMP4BoxSTSS>(new FMP4BoxSTSS(InParent, InBoxInfo)); }
+			virtual ~FMP4BoxSTSS() = default;
+			ELECTRABASE_API const TArray<uint32>& GetEntries();
+		protected:
+			FMP4BoxSTSS(const TWeakPtr<FMP4BoxBase>& InParent, const FMP4BoxInfo& InBoxInfo)
+				: FMP4BoxBase(InParent, InBoxInfo)
+			{ }
+			void ParseIfRequired();
+			struct FParsed;
+			TPimplPtr<FParsed> Parsed;
+		};
+
+
+		/**
+		 * `sdtp` box
+		 * ISO/IEC 14496-12:2022 - 8.6.4 Independent and disposable samples box
+		 */
+		class FMP4BoxSDTP : public FMP4BoxBase
+		{
+		public:
+			static TSharedPtr<FMP4BoxBase> Create(const TWeakPtr<FMP4BoxBase>& InParent, const FMP4BoxInfo& InBoxInfo)
+			{ return MakeShareable<FMP4BoxSDTP>(new FMP4BoxSDTP(InParent, InBoxInfo)); }
+			virtual ~FMP4BoxSDTP() = default;
+			ELECTRABASE_API TConstArrayView<uint8> GetEntries();
+		protected:
+			FMP4BoxSDTP(const TWeakPtr<FMP4BoxBase>& InParent, const FMP4BoxInfo& InBoxInfo)
+				: FMP4BoxBase(InParent, InBoxInfo)
+			{ }
+			void ParseIfRequired();
+			struct FParsed;
+			TPimplPtr<FParsed> Parsed;
+		};
+
+
+		/**
+		 * `stsc` box
+		 * ISO/IEC 14496-12:2022 - 8.7.4 Sample to chunk box
+		 */
+		class FMP4BoxSTSC : public FMP4BoxBase
+		{
+		public:
+			static TSharedPtr<FMP4BoxBase> Create(const TWeakPtr<FMP4BoxBase>& InParent, const FMP4BoxInfo& InBoxInfo)
+			{ return MakeShareable<FMP4BoxSTSC>(new FMP4BoxSTSC(InParent, InBoxInfo)); }
+			virtual ~FMP4BoxSTSC() = default;
+			struct FEntry
+			{
+				uint32 first_chunk = 0;
+				uint32 samples_per_chunk = 0;
+				uint32 sample_description_index = 0;
+			};
+			ELECTRABASE_API const TArray<FEntry>& GetEntries();
+		protected:
+			FMP4BoxSTSC(const TWeakPtr<FMP4BoxBase>& InParent, const FMP4BoxInfo& InBoxInfo)
+				: FMP4BoxBase(InParent, InBoxInfo)
+			{ }
+			void ParseIfRequired();
+			struct FParsed;
+			TPimplPtr<FParsed> Parsed;
+		};
+
+
+		/**
+		 * `stsz` box
+		 * ISO/IEC 14496-12:2022 - 8.7.3 Sample size boxes
+		 */
+		class FMP4BoxSTSZ : public FMP4BoxBase
+		{
+		public:
+			static TSharedPtr<FMP4BoxBase> Create(const TWeakPtr<FMP4BoxBase>& InParent, const FMP4BoxInfo& InBoxInfo)
+			{ return MakeShareable<FMP4BoxSTSZ>(new FMP4BoxSTSZ(InParent, InBoxInfo)); }
+			virtual ~FMP4BoxSTSZ() = default;
+			ELECTRABASE_API uint32 GetNumberOfSamples();
+			ELECTRABASE_API uint32 GetSizeOfSample(uint32 InIndex);
+		protected:
+			FMP4BoxSTSZ(const TWeakPtr<FMP4BoxBase>& InParent, const FMP4BoxInfo& InBoxInfo)
+				: FMP4BoxBase(InParent, InBoxInfo)
+			{ }
+			void ParseIfRequired();
+			struct FParsed;
+			TPimplPtr<FParsed> Parsed;
+		};
+
+
+		/**
+		 * `stco` box, `co64` box
+		 * ISO/IEC 14496-12:2022 - 8.7.5 Chunk offset box
+		 */
+		class FMP4BoxSTCO : public FMP4BoxBase
+		{
+		public:
+			static TSharedPtr<FMP4BoxBase> Create(const TWeakPtr<FMP4BoxBase>& InParent, const FMP4BoxInfo& InBoxInfo)
+			{ return MakeShareable<FMP4BoxSTCO>(new FMP4BoxSTCO(InParent, InBoxInfo)); }
+			virtual ~FMP4BoxSTCO() = default;
+			uint32 GetType() const override
+			{
+				// This class handles both `stco` and `co64`.
+				// For simplicities sake, when searching for the chunk offset box, we pretend this is
+				// an `stco` at all times, so user code does not need to worry about the difference.
+				return MakeBoxAtom('s','t','c','o');
+			}
+			ELECTRABASE_API uint32 GetNumberOfEntries();
+			ELECTRABASE_API uint64 GetChunkOffset(uint32 InIndex);
+		protected:
+			FMP4BoxSTCO(const TWeakPtr<FMP4BoxBase>& InParent, const FMP4BoxInfo& InBoxInfo)
+				: FMP4BoxBase(InParent, InBoxInfo)
+			{ }
+			void ParseIfRequired();
+			struct FParsed;
+			TPimplPtr<FParsed> Parsed;
+		};
+
+
+		/**
+		 * `saiz` box
+		 * ISO/IEC 14496-12:2022 - 8.7.8 Sample auxiliary information sizes box
+		 */
+		class FMP4BoxSAIZ : public FMP4BoxBase
+		{
+		public:
+			static TSharedPtr<FMP4BoxBase> Create(const TWeakPtr<FMP4BoxBase>& InParent, const FMP4BoxInfo& InBoxInfo)
+			{ return MakeShareable<FMP4BoxSAIZ>(new FMP4BoxSAIZ(InParent, InBoxInfo)); }
+			virtual ~FMP4BoxSAIZ() = default;
+			ELECTRABASE_API void Test();
+		protected:
+			FMP4BoxSAIZ(const TWeakPtr<FMP4BoxBase>& InParent, const FMP4BoxInfo& InBoxInfo)
+				: FMP4BoxBase(InParent, InBoxInfo)
+			{ }
+			void ParseIfRequired();
+			struct FParsed;
+			TPimplPtr<FParsed> Parsed;
+		};
+
+
+		/**
+		 * `saio` box
+		 * ISO/IEC 14496-12:2022 - 8.7.9 Sample auxiliary information offsets box
+		 */
+		class FMP4BoxSAIO : public FMP4BoxBase
+		{
+		public:
+			static TSharedPtr<FMP4BoxBase> Create(const TWeakPtr<FMP4BoxBase>& InParent, const FMP4BoxInfo& InBoxInfo)
+			{ return MakeShareable<FMP4BoxSAIO>(new FMP4BoxSAIO(InParent, InBoxInfo)); }
+			virtual ~FMP4BoxSAIO() = default;
+			ELECTRABASE_API void Test();
+		protected:
+			FMP4BoxSAIO(const TWeakPtr<FMP4BoxBase>& InParent, const FMP4BoxInfo& InBoxInfo)
+				: FMP4BoxBase(InParent, InBoxInfo)
+			{ }
+			void ParseIfRequired();
+			struct FParsed;
+			TPimplPtr<FParsed> Parsed;
+		};
+
+
+		/**
+		 * `sgpd` box
+		 * ISO/IEC 14496-12:2022 - 8.9.3 Sample group description box
+		 */
+		class FMP4BoxSGPD : public FMP4BoxBase
+		{
+		public:
+			static TSharedPtr<FMP4BoxBase> Create(const TWeakPtr<FMP4BoxBase>& InParent, const FMP4BoxInfo& InBoxInfo)
+			{ return MakeShareable<FMP4BoxSGPD>(new FMP4BoxSGPD(InParent, InBoxInfo)); }
+			virtual ~FMP4BoxSGPD() = default;
+
+			ELECTRABASE_API uint32 GetGroupingType();
+			ELECTRABASE_API const TArray<TConstArrayView<uint8>>& GetGroupDescriptionEntries();
+			ELECTRABASE_API uint32 GetDefaultGroupDescriptionIndex();
+		protected:
+			FMP4BoxSGPD(const TWeakPtr<FMP4BoxBase>& InParent, const FMP4BoxInfo& InBoxInfo)
+				: FMP4BoxBase(InParent, InBoxInfo)
+			{ }
+			void ParseIfRequired();
+			struct FParsed;
+			TPimplPtr<FParsed> Parsed;
+		};
+
+
+		/**
+		 * `sbgp` box
+		 * ISO/IEC 14496-12:2022 - 8.9.2 Sample to group box
+		 */
+		class FMP4BoxSBGP : public FMP4BoxBase
+		{
+		public:
+			static TSharedPtr<FMP4BoxBase> Create(const TWeakPtr<FMP4BoxBase>& InParent, const FMP4BoxInfo& InBoxInfo)
+			{ return MakeShareable<FMP4BoxSBGP>(new FMP4BoxSBGP(InParent, InBoxInfo)); }
+			virtual ~FMP4BoxSBGP() = default;
+
+			struct FEntry
+			{
+				uint32 sample_count = 0;
+				uint32 group_description_index = 0;
+			};
+
+			ELECTRABASE_API uint32 GetGroupingType();
+			ELECTRABASE_API uint32 GetGroupingTypeParameter();
+			ELECTRABASE_API const TArray<FEntry>& GetEntries();
+			ELECTRABASE_API uint32 GetNumTotalSamples();
+		protected:
+			FMP4BoxSBGP(const TWeakPtr<FMP4BoxBase>& InParent, const FMP4BoxInfo& InBoxInfo)
+				: FMP4BoxBase(InParent, InBoxInfo)
+			{ }
+			void ParseIfRequired();
+			struct FParsed;
+			TPimplPtr<FParsed> Parsed;
+		};
+
+
+		/**
+		 * `mvex` box
+		 * ISO/IEC 14496-12:2022 - 8.8.1 Movie extends box
+		 */
+		class FMP4BoxMVEX : public FMP4BoxBase
+		{
+		public:
+			static TSharedPtr<FMP4BoxBase> Create(const TWeakPtr<FMP4BoxBase>& InParent, const FMP4BoxInfo& InBoxInfo)
+			{ return MakeShareable<FMP4BoxMVEX>(new FMP4BoxMVEX(InParent, InBoxInfo)); }
+			virtual ~FMP4BoxMVEX() = default;
+			bool IsLeafBox() const override { return false; }
+		protected:
+			FMP4BoxMVEX(const TWeakPtr<FMP4BoxBase>& InParent, const FMP4BoxInfo& InBoxInfo)
+				: FMP4BoxBase(InParent, InBoxInfo)
+			{ }
+		};
+
+
+		/**
+		 * `mehd` box
+		 * ISO/IEC 14496-12:2022 - 8.8.2 Movie extends header box
+		 */
+		class FMP4BoxMEHD : public FMP4BoxBase
+		{
+		public:
+			static TSharedPtr<FMP4BoxBase> Create(const TWeakPtr<FMP4BoxBase>& InParent, const FMP4BoxInfo& InBoxInfo)
+			{ return MakeShareable<FMP4BoxMEHD>(new FMP4BoxMEHD(InParent, InBoxInfo)); }
+			virtual ~FMP4BoxMEHD() = default;
+			/**
+			 * Returns the fragment duration in the timescale of the `mvhd` box.
+			 * This is just the value as stored in the box. You need to apply the timescale yourself.
+			 */
+			ELECTRABASE_API uint64 GetFragmentDuration();
+		protected:
+			FMP4BoxMEHD(const TWeakPtr<FMP4BoxBase>& InParent, const FMP4BoxInfo& InBoxInfo)
+				: FMP4BoxBase(InParent, InBoxInfo)
+			{ }
+			void ParseIfRequired();
+			struct FParsed;
+			TPimplPtr<FParsed> Parsed;
+		};
+
+
+		/**
+		 * `trex` box
+		 * ISO/IEC 14496-12:2022 - 8.8.3 Track extends box
+		 */
+		class FMP4BoxTREX : public FMP4BoxBase
+		{
+		public:
+			static TSharedPtr<FMP4BoxBase> Create(const TWeakPtr<FMP4BoxBase>& InParent, const FMP4BoxInfo& InBoxInfo)
+			{ return MakeShareable<FMP4BoxTREX>(new FMP4BoxTREX(InParent, InBoxInfo)); }
+			virtual ~FMP4BoxTREX() = default;
+			ELECTRABASE_API uint32 GetTrackID();
+			ELECTRABASE_API uint32 GetDefaultSampleDescriptionIndex();
+			ELECTRABASE_API uint32 GetDefaultSampleDuration();
+			ELECTRABASE_API uint32 GetDefaultSampleSize();
+			ELECTRABASE_API uint32 GetDefaultSampleFlags();
+		protected:
+			FMP4BoxTREX(const TWeakPtr<FMP4BoxBase>& InParent, const FMP4BoxInfo& InBoxInfo)
+				: FMP4BoxBase(InParent, InBoxInfo)
+			{ }
+			void ParseIfRequired();
+			struct FParsed;
+			TPimplPtr<FParsed> Parsed;
+		};
+
+
+		/**
+		 * `leva` box
+		 * ISO/IEC 14496-12:2022 - 8.8.13 Level assignment box
+		 */
+		class FMP4BoxLEVA : public FMP4BoxBase
+		{
+		public:
+			static TSharedPtr<FMP4BoxBase> Create(const TWeakPtr<FMP4BoxBase>& InParent, const FMP4BoxInfo& InBoxInfo)
+			{ return MakeShareable<FMP4BoxLEVA>(new FMP4BoxLEVA(InParent, InBoxInfo)); }
+			virtual ~FMP4BoxLEVA() = default;
+		protected:
+			FMP4BoxLEVA(const TWeakPtr<FMP4BoxBase>& InParent, const FMP4BoxInfo& InBoxInfo)
+				: FMP4BoxBase(InParent, InBoxInfo)
+			{ }
+		};
+
+
+		/**
+		 * `moof` box
+		 * ISO/IEC 14496-12:2022 - 8.8.4 Movie fragment box
+		 */
+		class FMP4BoxMOOF : public FMP4BoxBase
+		{
+		public:
+			static TSharedPtr<FMP4BoxBase> Create(const TWeakPtr<FMP4BoxBase>& InParent, const FMP4BoxInfo& InBoxInfo)
+			{ return MakeShareable<FMP4BoxMOOF>(new FMP4BoxMOOF(InParent, InBoxInfo)); }
+			virtual ~FMP4BoxMOOF() = default;
+			bool IsLeafBox() const override { return false; }
+		protected:
+			FMP4BoxMOOF(const TWeakPtr<FMP4BoxBase>& InParent, const FMP4BoxInfo& InBoxInfo)
+				: FMP4BoxBase(InParent, InBoxInfo)
+			{ }
+		};
+
+
+		/**
+		 * `mfhd` box
+		 * ISO/IEC 14496-12:2022 - 8.8.5 Movie fragment header box
+		 */
+		class FMP4BoxMFHD : public FMP4BoxBase
+		{
+		public:
+			static TSharedPtr<FMP4BoxBase> Create(const TWeakPtr<FMP4BoxBase>& InParent, const FMP4BoxInfo& InBoxInfo)
+			{ return MakeShareable<FMP4BoxMFHD>(new FMP4BoxMFHD(InParent, InBoxInfo)); }
+			virtual ~FMP4BoxMFHD() = default;
+			ELECTRABASE_API uint32 GetSequenceNumber();
+		protected:
+			FMP4BoxMFHD(const TWeakPtr<FMP4BoxBase>& InParent, const FMP4BoxInfo& InBoxInfo)
+				: FMP4BoxBase(InParent, InBoxInfo)
+			{ }
+			void ParseIfRequired();
+			struct FParsed;
+			TPimplPtr<FParsed> Parsed;
+		};
+
+
+		/**
+		 * `traf` box
+		 * ISO/IEC 14496-12:2022 - 8.8.6 Track fragment box
+		 */
+		class FMP4BoxTRAF : public FMP4BoxBase
+		{
+		public:
+			static TSharedPtr<FMP4BoxBase> Create(const TWeakPtr<FMP4BoxBase>& InParent, const FMP4BoxInfo& InBoxInfo)
+			{ return MakeShareable<FMP4BoxTRAF>(new FMP4BoxTRAF(InParent, InBoxInfo)); }
+			virtual ~FMP4BoxTRAF() = default;
+			bool IsLeafBox() const override { return false; }
+		protected:
+			FMP4BoxTRAF(const TWeakPtr<FMP4BoxBase>& InParent, const FMP4BoxInfo& InBoxInfo)
+				: FMP4BoxBase(InParent, InBoxInfo)
+			{ }
+		};
+
+
+		/**
+		 * `tfhd` box
+		 * ISO/IEC 14496-12:2022 - 8.8.7 Track fragment header box
+		 */
+		class FMP4BoxTFHD : public FMP4BoxBase
+		{
+		public:
+			static TSharedPtr<FMP4BoxBase> Create(const TWeakPtr<FMP4BoxBase>& InParent, const FMP4BoxInfo& InBoxInfo)
+			{ return MakeShareable<FMP4BoxTFHD>(new FMP4BoxTFHD(InParent, InBoxInfo)); }
+			virtual ~FMP4BoxTFHD() = default;
+
+			ELECTRABASE_API uint32 GetTrackID();
+			ELECTRABASE_API bool HasBaseDataOffset();
+			ELECTRABASE_API uint64 GetBaseDataOffset();
+			ELECTRABASE_API bool HasSampleDescriptionIndex();
+			ELECTRABASE_API uint32 GetSampleDescriptionIndex();
+			ELECTRABASE_API bool HasDefaultSampleDuration();
+			ELECTRABASE_API uint32 GetDefaultSampleDuration();
+			ELECTRABASE_API bool HasDefaultSampleSize();
+			ELECTRABASE_API uint32 GetDefaultSampleSize();
+			ELECTRABASE_API bool HasDefaultSampleFlags();
+			ELECTRABASE_API uint32 GetDefaultSampleFlags();
+			ELECTRABASE_API bool IsDurationEmpty();
+			ELECTRABASE_API bool IsMoofDefaultBase();
+		protected:
+			FMP4BoxTFHD(const TWeakPtr<FMP4BoxBase>& InParent, const FMP4BoxInfo& InBoxInfo)
+				: FMP4BoxBase(InParent, InBoxInfo)
+			{ }
+			void ParseIfRequired();
+			struct FParsed;
+			TPimplPtr<FParsed> Parsed;
+		};
+
+
+		/**
+		 * `trun` box
+		 * ISO/IEC 14496-12:2022 - 8.8.8 Track fragment run box
+		 */
+		class FMP4BoxTRUN : public FMP4BoxBase
+		{
+		public:
+			static TSharedPtr<FMP4BoxBase> Create(const TWeakPtr<FMP4BoxBase>& InParent, const FMP4BoxInfo& InBoxInfo)
+			{ return MakeShareable<FMP4BoxTRUN>(new FMP4BoxTRUN(InParent, InBoxInfo)); }
+			virtual ~FMP4BoxTRUN() = default;
+			ELECTRABASE_API uint32 GetNumberOfSamples();
+			ELECTRABASE_API bool HasSampleOffset();
+			ELECTRABASE_API int32 GetSampleOffset();
+			ELECTRABASE_API bool HasFirstSampleFlags();
+			ELECTRABASE_API uint32 GetFirstSampleFlags();
+			ELECTRABASE_API bool HasSampleDurations();
+			ELECTRABASE_API const TArray<uint32>& GetSampleDurations();
+			ELECTRABASE_API bool HasSampleSizes();
+			ELECTRABASE_API const TArray<uint32>& GetSampleSizes();
+			ELECTRABASE_API bool HasSampleFlags();
+			ELECTRABASE_API const TArray<uint32>& GetSampleFlags();
+			ELECTRABASE_API bool HasSampleCompositionTimeOffsets();
+			ELECTRABASE_API const TArray<int32>& GetSampleCompositionTimeOffsets();
+		protected:
+			FMP4BoxTRUN(const TWeakPtr<FMP4BoxBase>& InParent, const FMP4BoxInfo& InBoxInfo)
+				: FMP4BoxBase(InParent, InBoxInfo)
+			{ }
+			void ParseIfRequired();
+			struct FParsed;
+			TPimplPtr<FParsed> Parsed;
+		};
+
+
+		/**
+		 * `mfra` box
+		 * ISO/IEC 14496-12:2022 - 8.8.9 Movie fragment random access box
+		 */
+		class FMP4BoxMFRA : public FMP4BoxBase
+		{
+		public:
+			static TSharedPtr<FMP4BoxBase> Create(const TWeakPtr<FMP4BoxBase>& InParent, const FMP4BoxInfo& InBoxInfo)
+			{ return MakeShareable<FMP4BoxMFRA>(new FMP4BoxMFRA(InParent, InBoxInfo)); }
+			virtual ~FMP4BoxMFRA() = default;
+			bool IsLeafBox() const override { return false; }
+		protected:
+			FMP4BoxMFRA(const TWeakPtr<FMP4BoxBase>& InParent, const FMP4BoxInfo& InBoxInfo)
+				: FMP4BoxBase(InParent, InBoxInfo)
+			{ }
+		};
+
+
+		/**
+		 * `tfra` box
+		 * ISO/IEC 14496-12:2022 - 8.8.10 Track fragment random access box
+		 */
+		class FMP4BoxTFRA : public FMP4BoxBase
+		{
+		public:
+			static TSharedPtr<FMP4BoxBase> Create(const TWeakPtr<FMP4BoxBase>& InParent, const FMP4BoxInfo& InBoxInfo)
+			{ return MakeShareable<FMP4BoxTFRA>(new FMP4BoxTFRA(InParent, InBoxInfo)); }
+			virtual ~FMP4BoxTFRA() = default;
+		protected:
+			FMP4BoxTFRA(const TWeakPtr<FMP4BoxBase>& InParent, const FMP4BoxInfo& InBoxInfo)
+				: FMP4BoxBase(InParent, InBoxInfo)
+			{ }
+		};
+
+
+		/**
+		 * `mfro` box
+		 * ISO/IEC 14496-12:2022 - 8.8.11 Movie fragment random access offset box
+		 */
+		class FMP4BoxMFRO : public FMP4BoxBase
+		{
+		public:
+			static TSharedPtr<FMP4BoxBase> Create(const TWeakPtr<FMP4BoxBase>& InParent, const FMP4BoxInfo& InBoxInfo)
+			{ return MakeShareable<FMP4BoxMFRO>(new FMP4BoxMFRO(InParent, InBoxInfo)); }
+			virtual ~FMP4BoxMFRO() = default;
+		protected:
+			FMP4BoxMFRO(const TWeakPtr<FMP4BoxBase>& InParent, const FMP4BoxInfo& InBoxInfo)
+				: FMP4BoxBase(InParent, InBoxInfo)
+			{ }
+		};
+
+
+		/**
+		 * `tfdt` box
+		 * ISO/IEC 14496-12:2022 - 8.8.12 Track fragment decode time box
+		 */
+		class FMP4BoxTFDT : public FMP4BoxBase
+		{
+		public:
+			static TSharedPtr<FMP4BoxBase> Create(const TWeakPtr<FMP4BoxBase>& InParent, const FMP4BoxInfo& InBoxInfo)
+			{ return MakeShareable<FMP4BoxTFDT>(new FMP4BoxTFDT(InParent, InBoxInfo)); }
+			virtual ~FMP4BoxTFDT() = default;
+			ELECTRABASE_API uint64 GetBaseMediaDecodeTime();
+		protected:
+			FMP4BoxTFDT(const TWeakPtr<FMP4BoxBase>& InParent, const FMP4BoxInfo& InBoxInfo)
+				: FMP4BoxBase(InParent, InBoxInfo)
+			{ }
+			void ParseIfRequired();
+			struct FParsed;
+			TPimplPtr<FParsed> Parsed;
+		};
+
+
+		/**
+		 * `sidx` box
+		 * ISO/IEC 14496-12:2022 - 8.16.3 Segment index box
+		 */
+		class FMP4BoxSIDX : public FMP4BoxBase
+		{
+		public:
+			static TSharedPtr<FMP4BoxBase> Create(const TWeakPtr<FMP4BoxBase>& InParent, const FMP4BoxInfo& InBoxInfo)
+			{ return MakeShareable<FMP4BoxSIDX>(new FMP4BoxSIDX(InParent, InBoxInfo)); }
+			virtual ~FMP4BoxSIDX() = default;
+			struct FEntry
+			{
+				uint32 SubSegmentDuration;
+				uint32 IsReferenceType : 1;
+				uint32 Size : 31;
+				uint32 StartsWithSAP : 1;
+				uint32 SAPType : 3;
+				uint32 SAPDeltaTime : 28;
+			};
+			using FEntryList = TChunkedArray<FEntry, 3000>;
+			ELECTRABASE_API uint32 GetReferenceID();
+			ELECTRABASE_API uint32 GetTimescale();
+			ELECTRABASE_API uint64 GetEarliestPresentationTime();
+			ELECTRABASE_API uint64 GetFirstOffset();
+			ELECTRABASE_API const FEntryList& GetEntries();
+		protected:
+			FMP4BoxSIDX(const TWeakPtr<FMP4BoxBase>& InParent, const FMP4BoxInfo& InBoxInfo)
+				: FMP4BoxBase(InParent, InBoxInfo)
+			{ }
+			void ParseIfRequired();
+			struct FParsed;
+			TPimplPtr<FParsed> Parsed;
+		};
+
+
+		/**
+		 * `ssix` box
+		 * ISO/IEC 14496-12:2022 - 8.16.4 Subsegment index box
+		 */
+		class FMP4BoxSSIX : public FMP4BoxBase
+		{
+		public:
+			static TSharedPtr<FMP4BoxBase> Create(const TWeakPtr<FMP4BoxBase>& InParent, const FMP4BoxInfo& InBoxInfo)
+			{ return MakeShareable<FMP4BoxSSIX>(new FMP4BoxSSIX(InParent, InBoxInfo)); }
+			virtual ~FMP4BoxSSIX() = default;
+		protected:
+			FMP4BoxSSIX(const TWeakPtr<FMP4BoxBase>& InParent, const FMP4BoxInfo& InBoxInfo)
+				: FMP4BoxBase(InParent, InBoxInfo)
+			{ }
+		};
+
+
+		/**
+		 * `vmhd` box
+		 * ISO/IEC 14496-12:2022 - 12.1.2 Video media header
+		 */
+		class FMP4BoxVMHD : public FMP4BoxBase
+		{
+		public:
+			static TSharedPtr<FMP4BoxBase> Create(const TWeakPtr<FMP4BoxBase>& InParent, const FMP4BoxInfo& InBoxInfo)
+			{ return MakeShareable<FMP4BoxVMHD>(new FMP4BoxVMHD(InParent, InBoxInfo)); }
+			virtual ~FMP4BoxVMHD() = default;
+		protected:
+			FMP4BoxVMHD(const TWeakPtr<FMP4BoxBase>& InParent, const FMP4BoxInfo& InBoxInfo)
+				: FMP4BoxBase(InParent, InBoxInfo)
+			{ }
+		};
+
+
+		/**
+		 * `smhd` box
+		 * ISO/IEC 14496-12:2022 - 12.2.2 Sound media header
+		 */
+		class FMP4BoxSMHD : public FMP4BoxBase
+		{
+		public:
+			static TSharedPtr<FMP4BoxBase> Create(const TWeakPtr<FMP4BoxBase>& InParent, const FMP4BoxInfo& InBoxInfo)
+			{ return MakeShareable<FMP4BoxSMHD>(new FMP4BoxSMHD(InParent, InBoxInfo)); }
+			virtual ~FMP4BoxSMHD() = default;
+		protected:
+			FMP4BoxSMHD(const TWeakPtr<FMP4BoxBase>& InParent, const FMP4BoxInfo& InBoxInfo)
+				: FMP4BoxBase(InParent, InBoxInfo)
+			{ }
+		};
+
+
+		/**
+		 * `nmhd` box
+		 * ISO/IEC 14496-12:2022 - 8.4.5.2 Null media header box
+		 */
+		class FMP4BoxNMHD : public FMP4BoxBase
+		{
+		public:
+			static TSharedPtr<FMP4BoxBase> Create(const TWeakPtr<FMP4BoxBase>& InParent, const FMP4BoxInfo& InBoxInfo)
+			{ return MakeShareable<FMP4BoxNMHD>(new FMP4BoxNMHD(InParent, InBoxInfo)); }
+			virtual ~FMP4BoxNMHD() = default;
+		protected:
+			FMP4BoxNMHD(const TWeakPtr<FMP4BoxBase>& InParent, const FMP4BoxInfo& InBoxInfo)
+				: FMP4BoxBase(InParent, InBoxInfo)
+			{ }
+		};
+
+
+		/**
+		 * `gmhd` box
+		 * Apple QuickTime: https://developer.apple.com/documentation/quicktime-file-format/base_media_information_header_atom
+		 */
+		class FMP4BoxGMHD : public FMP4BoxBase
+		{
+		public:
+			static TSharedPtr<FMP4BoxBase> Create(const TWeakPtr<FMP4BoxBase>& InParent, const FMP4BoxInfo& InBoxInfo)
+			{ return MakeShareable<FMP4BoxGMHD>(new FMP4BoxGMHD(InParent, InBoxInfo)); }
+			virtual ~FMP4BoxGMHD() = default;
+		protected:
+			FMP4BoxGMHD(const TWeakPtr<FMP4BoxBase>& InParent, const FMP4BoxInfo& InBoxInfo)
+				: FMP4BoxBase(InParent, InBoxInfo)
+			{ }
+		};
+
+
+		/**
+		 * `udta` box
+		 * ISO/IEC 14496-12:2022 - 8.10.1 User data box
+		 */
+		class FMP4BoxUDTA : public FMP4BoxBase
+		{
+		public:
+			static TSharedPtr<FMP4BoxBase> Create(const TWeakPtr<FMP4BoxBase>& InParent, const FMP4BoxInfo& InBoxInfo)
+			{ return MakeShareable<FMP4BoxUDTA>(new FMP4BoxUDTA(InParent, InBoxInfo)); }
+			virtual ~FMP4BoxUDTA() = default;
+			bool IsLeafBox() const override { return false; }
+		protected:
+			FMP4BoxUDTA(const TWeakPtr<FMP4BoxBase>& InParent, const FMP4BoxInfo& InBoxInfo)
+				: FMP4BoxBase(InParent, InBoxInfo)
+			{ }
+		};
+
+
+		/**
+		 * `meta` box
+		 * ISO/IEC 14496-12:2022 - 8.11.1 MetaBox
+		 */
+		class FMP4BoxMETA : public FMP4BoxBase
+		{
+		public:
+			static TSharedPtr<FMP4BoxBase> Create(const TWeakPtr<FMP4BoxBase>& InParent, const FMP4BoxInfo& InBoxInfo)
+			{ return MakeShareable<FMP4BoxMETA>(new FMP4BoxMETA(InParent, InBoxInfo)); }
+			virtual ~FMP4BoxMETA() = default;
+			bool IsLeafBox() const override { return false; }
+			bool IsListOfEntries() const override { return true; }
+		protected:
+			FMP4BoxMETA(const TWeakPtr<FMP4BoxBase>& InParent, const FMP4BoxInfo& InBoxInfo)
+				: FMP4BoxBase(InParent, InBoxInfo)
+			{ }
+		};
+
+
+		/**
+		 * `sinf` box
+		 * ISO/IEC 14496-12:2022 - 8.12.2 Protection scheme information box
+		 */
+		class FMP4BoxSINF : public FMP4BoxBase
+		{
+		public:
+			static TSharedPtr<FMP4BoxBase> Create(const TWeakPtr<FMP4BoxBase>& InParent, const FMP4BoxInfo& InBoxInfo)
+			{ return MakeShareable<FMP4BoxSINF>(new FMP4BoxSINF(InParent, InBoxInfo)); }
+			virtual ~FMP4BoxSINF() = default;
+			bool IsLeafBox() const override { return false; }
+		protected:
+			FMP4BoxSINF(const TWeakPtr<FMP4BoxBase>& InParent, const FMP4BoxInfo& InBoxInfo)
+				: FMP4BoxBase(InParent, InBoxInfo)
+			{ }
+		};
+
+
+		/**
+		 * `frma` box
+		 * ISO/IEC 14496-12:2022 - 8.12.3 Original format box
+		 */
+		class FMP4BoxFRMA : public FMP4BoxBase
+		{
+		public:
+			static TSharedPtr<FMP4BoxBase> Create(const TWeakPtr<FMP4BoxBase>& InParent, const FMP4BoxInfo& InBoxInfo)
+			{ return MakeShareable<FMP4BoxFRMA>(new FMP4BoxFRMA(InParent, InBoxInfo)); }
+			virtual ~FMP4BoxFRMA() = default;
+		protected:
+			FMP4BoxFRMA(const TWeakPtr<FMP4BoxBase>& InParent, const FMP4BoxInfo& InBoxInfo)
+				: FMP4BoxBase(InParent, InBoxInfo)
+			{ }
+		};
+
+
+		/**
+		 * `schm` box
+		 * ISO/IEC 14496-12:2022 - 8.12.6 Scheme type box
+		 */
+		class FMP4BoxSCHM : public FMP4BoxBase
+		{
+		public:
+			static TSharedPtr<FMP4BoxBase> Create(const TWeakPtr<FMP4BoxBase>& InParent, const FMP4BoxInfo& InBoxInfo)
+			{ return MakeShareable<FMP4BoxSCHM>(new FMP4BoxSCHM(InParent, InBoxInfo)); }
+			virtual ~FMP4BoxSCHM() = default;
+		protected:
+			FMP4BoxSCHM(const TWeakPtr<FMP4BoxBase>& InParent, const FMP4BoxInfo& InBoxInfo)
+				: FMP4BoxBase(InParent, InBoxInfo)
+			{ }
+		};
+
+
+		/**
+		 * `schi` box
+		 * ISO/IEC 14496-12:2022 - 8.12.7 Scheme information box
+		 */
+		class FMP4BoxSCHI : public FMP4BoxBase
+		{
+		public:
+			static TSharedPtr<FMP4BoxBase> Create(const TWeakPtr<FMP4BoxBase>& InParent, const FMP4BoxInfo& InBoxInfo)
+			{ return MakeShareable<FMP4BoxSCHI>(new FMP4BoxSCHI(InParent, InBoxInfo)); }
+			virtual ~FMP4BoxSCHI() = default;
+			bool IsLeafBox() const override { return false; }
+		protected:
+			FMP4BoxSCHI(const TWeakPtr<FMP4BoxBase>& InParent, const FMP4BoxInfo& InBoxInfo)
+				: FMP4BoxBase(InParent, InBoxInfo)
+			{ }
+		};
+
+
+		/**
+		 * `pssh` box
+		 * ISO/IEC 23001-7:2023 - 8.1 Protection system specific header box
+		 */
+		class FMP4BoxPSSH : public FMP4BoxBase
+		{
+		public:
+			static TSharedPtr<FMP4BoxBase> Create(const TWeakPtr<FMP4BoxBase>& InParent, const FMP4BoxInfo& InBoxInfo)
+			{ return MakeShareable<FMP4BoxPSSH>(new FMP4BoxPSSH(InParent, InBoxInfo)); }
+			virtual ~FMP4BoxPSSH() = default;
+			ELECTRABASE_API TConstArrayView<uint8> GetSystemID();
+			ELECTRABASE_API const TArray<TConstArrayView<uint8>>& GetKIDs();
+			ELECTRABASE_API TConstArrayView<uint8> GetData();
+		protected:
+			FMP4BoxPSSH(const TWeakPtr<FMP4BoxBase>& InParent, const FMP4BoxInfo& InBoxInfo)
+				: FMP4BoxBase(InParent, InBoxInfo)
+			{ }
+			void ParseIfRequired();
+			struct FParsed;
+			TPimplPtr<FParsed> Parsed;
+		};
+
+
+		/**
+		 * `tenc` box
+		 * ISO/IEC 23001-7:2023 - 8.2 Track Encryption box
+		 */
+		class FMP4BoxTENC : public FMP4BoxBase
+		{
+		public:
+			static TSharedPtr<FMP4BoxBase> Create(const TWeakPtr<FMP4BoxBase>& InParent, const FMP4BoxInfo& InBoxInfo)
+			{ return MakeShareable<FMP4BoxTENC>(new FMP4BoxTENC(InParent, InBoxInfo)); }
+			virtual ~FMP4BoxTENC() = default;
+			ELECTRABASE_API bool HasDefaultCryptBlockValues();
+			ELECTRABASE_API uint8 GetDefaultCryptByteBlock();
+			ELECTRABASE_API uint8 GetDefaultSkipByteBlock();
+			ELECTRABASE_API uint8 GetDefaultIsProtected();
+			ELECTRABASE_API uint8 GetDefaultPerSampleIVSize();
+			ELECTRABASE_API TConstArrayView<uint8> GetDefaultKID();
+			ELECTRABASE_API TConstArrayView<uint8> GetDefaultConstantIV();
+		protected:
+			FMP4BoxTENC(const TWeakPtr<FMP4BoxBase>& InParent, const FMP4BoxInfo& InBoxInfo)
+				: FMP4BoxBase(InParent, InBoxInfo)
+			{ }
+			void ParseIfRequired();
+			struct FParsed;
+			TPimplPtr<FParsed> Parsed;
+		};
+
+
+		/**
+		 * `senc` box
+		 * ISO/IEC 23001-7:2023 - 7.2.1 Sample encryption box - Definition
+		 */
+		class FMP4BoxSENC : public FMP4BoxBase
+		{
+		public:
+			static TSharedPtr<FMP4BoxBase> Create(const TWeakPtr<FMP4BoxBase>& InParent, const FMP4BoxInfo& InBoxInfo)
+			{ return MakeShareable<FMP4BoxSENC>(new FMP4BoxSENC(InParent, InBoxInfo)); }
+			virtual ~FMP4BoxSENC() = default;
+
+			ELECTRABASE_API void Prepare(const TArray<TSharedPtr<Electra::UtilitiesMP4::FMP4BoxBase>>& InRelatedBoxes);
+
+			struct FSubSample
+			{
+				uint32 NumEncryptedBytes = 0;
+				uint16 NumClearBytes = 0;
+			};
+			struct FEntry
+			{
+				TArray<uint8> IV;
+				TArray<FSubSample> SubSamples;
+			};
+			ELECTRABASE_API const TArray<FEntry>& GetEntries();
+		protected:
+			FMP4BoxSENC(const TWeakPtr<FMP4BoxBase>& InParent, const FMP4BoxInfo& InBoxInfo)
+				: FMP4BoxBase(InParent, InBoxInfo)
+			{ }
+			void ParseIfRequired();
+			struct FParsed;
+			TPimplPtr<FParsed> Parsed;
+			struct FInitRefs;
+			TPimplPtr<FInitRefs> InitRefs;
+		};
+
+
+		/**
+		 * `avcC` box
+		 * ISO/IEC 14496-15:2022 - 5.4.2 AVC video stream definition
+		 */
+		class FMP4BoxAVCC : public FMP4BoxBase
+		{
+		public:
+			static TSharedPtr<FMP4BoxBase> Create(const TWeakPtr<FMP4BoxBase>& InParent, const FMP4BoxInfo& InBoxInfo)
+			{ return MakeShareable<FMP4BoxAVCC>(new FMP4BoxAVCC(InParent, InBoxInfo)); }
+			virtual ~FMP4BoxAVCC() = default;
+			ELECTRABASE_API TConstArrayView<uint8> GetAVCDecoderConfigurationRecord();
+		protected:
+			FMP4BoxAVCC(const TWeakPtr<FMP4BoxBase>& InParent, const FMP4BoxInfo& InBoxInfo)
+				: FMP4BoxBase(InParent, InBoxInfo)
+			{ }
+			void ParseIfRequired();
+			struct FParsed;
+			TPimplPtr<FParsed> Parsed;
+		};
+
+
+		/**
+		 * `hvcC` box
+		 * ISO/IEC 14496-15:2022 - 5.4.2 AVC video stream definition
+		 */
+		class FMP4BoxHVCC : public FMP4BoxBase
+		{
+		public:
+			static TSharedPtr<FMP4BoxBase> Create(const TWeakPtr<FMP4BoxBase>& InParent, const FMP4BoxInfo& InBoxInfo)
+			{ return MakeShareable<FMP4BoxHVCC>(new FMP4BoxHVCC(InParent, InBoxInfo)); }
+			virtual ~FMP4BoxHVCC() = default;
+			ELECTRABASE_API TConstArrayView<uint8> GetHEVCDecoderConfigurationRecord();
+		protected:
+			FMP4BoxHVCC(const TWeakPtr<FMP4BoxBase>& InParent, const FMP4BoxInfo& InBoxInfo)
+				: FMP4BoxBase(InParent, InBoxInfo)
+			{ }
+			void ParseIfRequired();
+			struct FParsed;
+			TPimplPtr<FParsed> Parsed;
+		};
+
+
+		/**
+		 * `iods` box
+		 * ISO/IEC 14496-14:2020 - 6.2 Object Descriptor Box
+		 */
+		class FMP4BoxIODS : public FMP4BoxBase
+		{
+		public:
+			static TSharedPtr<FMP4BoxBase> Create(const TWeakPtr<FMP4BoxBase>& InParent, const FMP4BoxInfo& InBoxInfo)
+			{ return MakeShareable<FMP4BoxIODS>(new FMP4BoxIODS(InParent, InBoxInfo)); }
+			virtual ~FMP4BoxIODS() = default;
+			ELECTRABASE_API TConstArrayView<uint8> GetObjectDescriptor();
+		protected:
+			FMP4BoxIODS(const TWeakPtr<FMP4BoxBase>& InParent, const FMP4BoxInfo& InBoxInfo)
+				: FMP4BoxBase(InParent, InBoxInfo)
+			{ }
+			void ParseIfRequired();
+			struct FParsed;
+			TPimplPtr<FParsed> Parsed;
+		};
+
+
+		/**
+		 * `esds` box
+		 * ISO/IEC 14496-14:2020 - 6.7.2
+		 */
+		class FMP4BoxESDS : public FMP4BoxBase
+		{
+		public:
+			static TSharedPtr<FMP4BoxBase> Create(const TWeakPtr<FMP4BoxBase>& InParent, const FMP4BoxInfo& InBoxInfo)
+			{ return MakeShareable<FMP4BoxESDS>(new FMP4BoxESDS(InParent, InBoxInfo)); }
+			virtual ~FMP4BoxESDS() = default;
+			ELECTRABASE_API TConstArrayView<uint8> GetESDescriptor();
+		protected:
+			FMP4BoxESDS(const TWeakPtr<FMP4BoxBase>& InParent, const FMP4BoxInfo& InBoxInfo)
+				: FMP4BoxBase(InParent, InBoxInfo)
+			{ }
+			void ParseIfRequired();
+			struct FParsed;
+			TPimplPtr<FParsed> Parsed;
+		};
+
+
+		/**
+		 * `dvcC` box, `dvvC` box, `dvwC` box
+		 * Dolby Vision Streams Within the ISO Base:2023 - 2.2 Dolby Vision configuration boxes
+		 */
+		class FMP4BoxDVCC : public FMP4BoxBase
+		{
+		public:
+			static TSharedPtr<FMP4BoxBase> Create(const TWeakPtr<FMP4BoxBase>& InParent, const FMP4BoxInfo& InBoxInfo)
+			{ return MakeShareable<FMP4BoxDVCC>(new FMP4BoxDVCC(InParent, InBoxInfo)); }
+			virtual ~FMP4BoxDVCC() = default;
+			ELECTRABASE_API TConstArrayView<uint8> GetDOVIDecoderConfigurationRecord();
+		protected:
+			FMP4BoxDVCC(const TWeakPtr<FMP4BoxBase>& InParent, const FMP4BoxInfo& InBoxInfo)
+				: FMP4BoxBase(InParent, InBoxInfo)
+			{ }
+			void ParseIfRequired();
+			struct FParsed;
+			TPimplPtr<FParsed> Parsed;
+		};
+
+
+		/**
+		 * `dac3` box
+		 * Annex F.4 of ETSI TS 102 366 - AC3SpecificBox
+		 */
+		class FMP4BoxDAC3 : public FMP4BoxBase
+		{
+		public:
+			static TSharedPtr<FMP4BoxBase> Create(const TWeakPtr<FMP4BoxBase>& InParent, const FMP4BoxInfo& InBoxInfo)
+			{ return MakeShareable<FMP4BoxDAC3>(new FMP4BoxDAC3(InParent, InBoxInfo)); }
+			virtual ~FMP4BoxDAC3() = default;
+			ELECTRABASE_API TConstArrayView<uint8> GetAC3SpecificBox();
+		protected:
+			FMP4BoxDAC3(const TWeakPtr<FMP4BoxBase>& InParent, const FMP4BoxInfo& InBoxInfo)
+				: FMP4BoxBase(InParent, InBoxInfo)
+			{ }
+			void ParseIfRequired();
+			struct FParsed;
+			TPimplPtr<FParsed> Parsed;
+		};
+
+
+		/**
+		 * `dec3` box
+		 * Annex F.6 of ETSI TS 102 366 - EC3SpecificBox
+		 */
+		class FMP4BoxDEC3 : public FMP4BoxBase
+		{
+		public:
+			static TSharedPtr<FMP4BoxBase> Create(const TWeakPtr<FMP4BoxBase>& InParent, const FMP4BoxInfo& InBoxInfo)
+			{ return MakeShareable<FMP4BoxDEC3>(new FMP4BoxDEC3(InParent, InBoxInfo)); }
+			virtual ~FMP4BoxDEC3() = default;
+			ELECTRABASE_API TConstArrayView<uint8> GetEC3SpecificBox();
+		protected:
+			FMP4BoxDEC3(const TWeakPtr<FMP4BoxBase>& InParent, const FMP4BoxInfo& InBoxInfo)
+				: FMP4BoxBase(InParent, InBoxInfo)
+			{ }
+			void ParseIfRequired();
+			struct FParsed;
+			TPimplPtr<FParsed> Parsed;
+		};
+
+
+		/**
+		 * `dfLa` box
+		 * FLAC ([REDACTED_PROFILE_URL]/flac/blob/master/doc/isoflac.txt) specific box - 3.3.2 FLAC Specific Box
+		 */
+		class FMP4BoxDFLA : public FMP4BoxBase
+		{
+		public:
+			static TSharedPtr<FMP4BoxBase> Create(const TWeakPtr<FMP4BoxBase>& InParent, const FMP4BoxInfo& InBoxInfo)
+			{ return MakeShareable<FMP4BoxDFLA>(new FMP4BoxDFLA(InParent, InBoxInfo)); }
+			virtual ~FMP4BoxDFLA() = default;
+			ELECTRABASE_API TConstArrayView<uint8> GetFLACSpecificBox();
+		protected:
+			FMP4BoxDFLA(const TWeakPtr<FMP4BoxBase>& InParent, const FMP4BoxInfo& InBoxInfo)
+				: FMP4BoxBase(InParent, InBoxInfo)
+			{ }
+			void ParseIfRequired();
+			struct FParsed;
+			TPimplPtr<FParsed> Parsed;
+		};
+
+
+		/**
+		 * `dOps` box
+		 * Opus (https://opus-codec.org/docs/opus_in_isobmff.html#4.3.2) specific box - 4.3.2 Opus Specific Bo
+		 */
+		class FMP4BoxDOPS : public FMP4BoxBase
+		{
+		public:
+			static TSharedPtr<FMP4BoxBase> Create(const TWeakPtr<FMP4BoxBase>& InParent, const FMP4BoxInfo& InBoxInfo)
+			{ return MakeShareable<FMP4BoxDOPS>(new FMP4BoxDOPS(InParent, InBoxInfo)); }
+			virtual ~FMP4BoxDOPS() = default;
+			ELECTRABASE_API TConstArrayView<uint8> GetOpusSpecificBox();
+		protected:
+			FMP4BoxDOPS(const TWeakPtr<FMP4BoxBase>& InParent, const FMP4BoxInfo& InBoxInfo)
+				: FMP4BoxBase(InParent, InBoxInfo)
+			{ }
+			void ParseIfRequired();
+			struct FParsed;
+			TPimplPtr<FParsed> Parsed;
+		};
+
+
+		/**
+		 * `vpcC` box
+		 * WebM VP8/VP9 (https://www.webmproject.org/vp9/mp4/) specific box - VPCodecConfigurationBox
+		 */
+		class FMP4BoxVPCC : public FMP4BoxBase
+		{
+		public:
+			static TSharedPtr<FMP4BoxBase> Create(const TWeakPtr<FMP4BoxBase>& InParent, const FMP4BoxInfo& InBoxInfo)
+			{ return MakeShareable<FMP4BoxVPCC>(new FMP4BoxVPCC(InParent, InBoxInfo)); }
+			virtual ~FMP4BoxVPCC() = default;
+			ELECTRABASE_API TConstArrayView<uint8> GetVPCodecConfigurationBox();
+		protected:
+			FMP4BoxVPCC(const TWeakPtr<FMP4BoxBase>& InParent, const FMP4BoxInfo& InBoxInfo)
+				: FMP4BoxBase(InParent, InBoxInfo)
+			{ }
+			void ParseIfRequired();
+			struct FParsed;
+			TPimplPtr<FParsed> Parsed;
+		};
+
+
+		/**
+		 * `wave` box
+		 * Apple Quicktime: https://developer.apple.com/documentation/quicktime-file-format/sidecompressionparam_atom
+		 */
+		class FMP4BoxWAVE : public FMP4BoxBase
+		{
+		public:
+			static TSharedPtr<FMP4BoxBase> Create(const TWeakPtr<FMP4BoxBase>& InParent, const FMP4BoxInfo& InBoxInfo)
+			{ return MakeShareable<FMP4BoxWAVE>(new FMP4BoxWAVE(InParent, InBoxInfo)); }
+			virtual ~FMP4BoxWAVE() = default;
+		protected:
+			FMP4BoxWAVE(const TWeakPtr<FMP4BoxBase>& InParent, const FMP4BoxInfo& InBoxInfo)
+				: FMP4BoxBase(InParent, InBoxInfo)
+			{ }
+		};
+
+
+		/**
+		 * `tapt` box
+		 * Apple Quicktime: https://developer.apple.com/documentation/quicktime-file-format/track_aperture_mode_dimensions_atom
+		 */
+		class FMP4BoxTAPT : public FMP4BoxBase
+		{
+		public:
+			static TSharedPtr<FMP4BoxBase> Create(const TWeakPtr<FMP4BoxBase>& InParent, const FMP4BoxInfo& InBoxInfo)
+			{ return MakeShareable<FMP4BoxTAPT>(new FMP4BoxTAPT(InParent, InBoxInfo)); }
+			virtual ~FMP4BoxTAPT() = default;
+			bool IsLeafBox() const override { return false; }
+		protected:
+			FMP4BoxTAPT(const TWeakPtr<FMP4BoxBase>& InParent, const FMP4BoxInfo& InBoxInfo)
+				: FMP4BoxBase(InParent, InBoxInfo)
+			{ }
+		};
+
+
+		class FMP4BoxFactory
+		{
+		public:
+			static ELECTRABASE_API FMP4BoxFactory& Get();
+			ELECTRABASE_API TSharedPtr<FMP4BoxBase> Create(const TWeakPtr<FMP4BoxBase>& InParent, const FMP4BoxInfo& InBoxInfo);
+		private:
+			FMP4BoxFactory();
+			TMap<uint32, TSharedPtr<FMP4BoxBase>(*)(const TWeakPtr<FMP4BoxBase>& /*InParent*/, const FMP4BoxInfo& /*InBoxInfo*/)> FactoryMap;
+		};
+
+	} // namespace UtilitiesMP4
+
+} // namespace Electra
diff --git a/Engine/Plugins/Media/ElectraUtil/Source/ElectraBase/Public/Utilities/MP4Boxes/MP4Track.h b/Engine/Plugins/Media/ElectraUtil/Source/ElectraBase/Public/Utilities/MP4Boxes/MP4Track.h
new file mode 100644
--- /dev/null
+++ b/Engine/Plugins/Media/ElectraUtil/Source/ElectraBase/Public/Utilities/MP4Boxes/MP4Track.h
@@ -0,0 +1,215 @@
+// Copyright Epic Games, Inc. All Rights Reserved.
+
+#pragma once
+
+#include "Utilities/UtilitiesMP4.h"
+#include "Utilities/MP4Boxes/MP4Boxes.h"
+#include "Utilities/MP4Boxes/MP4BoxIterators.h"
+#include "Utilities/MP4Boxes/MP4BoxMetadata.h"
+
+namespace Electra
+{
+	namespace UtilitiesMP4
+	{
+		class FMP4Track : public TSharedFromThis<FMP4Track>
+		{
+		public:
+			struct FFirstSample
+			{
+				int64 SamplePTS = 0;
+				int64 StartPTS = 0;
+				uint32 SampleNumber = 0;
+				uint32 SyncSampleNumber = 0;
+			};
+
+			struct FLastSample
+			{
+				int64 SamplePTS = 0;
+				int64 EndPTS = 0;
+				uint32 SampleNumber = 0;
+				uint32 LastSampleNumber = 0;
+			};
+
+		protected:
+			struct FConvenience
+			{
+				FFirstSample FirstSample;
+				FLastSample LastSample;
+				FTimeFraction DurationFromMvhdBox;
+				FTimeFraction DurationFromTkhdBox;
+				FTimeFraction DurationFromMdhdBox;
+				FTimeFraction MappedDurationFromElstBox;
+				int64 BaseMediaDecodeTime = 0;
+				int64 CompositionTimeAtZeroPoint = 0;
+				int64 DTSShiftAtZeroPoint = 0;
+				uint32 TrackID = 0;
+				uint32 NumTotalSamples = 0;
+			};
+
+		public:
+			static TSharedPtr<FMP4Track> Create(const TSharedPtr<FMP4BoxTRAK>& InTrakBox)
+			{ return MakeShareable<FMP4Track>(new FMP4Track(InTrakBox)); }
+			virtual ~FMP4Track() = default;
+
+			ELECTRABASE_API bool Prepare(FTimeFraction InMovieDurationFromMVHDBox);
+			ELECTRABASE_API FString GetLastError();
+			ELECTRABASE_API const FMP4TrackMetadataCommon& GetCommonMetadata();
+
+			class FIterator
+			{
+			public:
+				virtual ~FIterator() = default;
+
+				bool IsValid() const
+				{ return Convs.NumTotalSamples != 0; }
+				TWeakPtr<FMP4Track> GetTrack() const
+				{ return Track; }
+				uint32 GetTrackID() const
+				{ return Convs.TrackID; }
+				uint32 GetSampleNumber() const
+				{ return SampleNumber; }
+				// Returns the DTS without mapping to the timeline.
+				FTimeFraction GetDTS() const
+				{ return CurrentDTS; }
+				// Returns the effective DTS, which has the timeline mapping applied. This may result in a negative value.
+				FTimeFraction GetEffectiveDTS() const
+				{ return CurrentEffectiveDTS; }
+				// Returns the PTS as the sum of the DTS and the composition time offset, without mapping to the timeline.
+				FTimeFraction GetPTS() const
+				{ return CurrentPTS; }
+				// Returns the effective PTS, which is the media time mapped into the 0-based timeline.
+				FTimeFraction GetEffectivePTS() const
+				{ return CurrentEffectivePTS; }
+				FTimeFraction GetDuration() const
+				{ return CurrentDuration; }
+				bool IsSyncOrRAPSample() const
+				{ return bCurrentIsSyncOrRAP; }
+				int64 GetSampleSize() const
+				{ return CurrentSampleSize; }
+				int64 GetSampleFileOffset() const
+				{ return CurrentSampleFileOffset; }
+				uint32 GetTimescale() const
+				{ return Convs.DurationFromMdhdBox.GetDenominator(); }
+				uint32 GetNumSamples() const
+				{ return Convs.NumTotalSamples; }
+				// Returns the track's entire media duration, not affected by an edit list. Timescale comes from `mdhd` box.
+				FTimeFraction GetTrackDuration() const
+				{ return Convs.DurationFromMdhdBox; }
+				// Returns the effective track's duration, as specified by an edit list. Timescale has been converted into `mdhd` timescale!
+				FTimeFraction GetEffectiveTrackDuration() const
+				{ return Convs.MappedDurationFromElstBox; }
+
+				// Returns information about the first sample that is mapped to the 0-based timeline via `elst` box.
+				const FFirstSample& GetFirstSampleInfo() const
+				{ return Convs.FirstSample; }
+				// Returns information about the last sample that is mapped to the 0-based timeline via `elst` box.
+				const FLastSample& GetLastSampleInfo() const
+				{ return Convs.LastSample; }
+
+				// Advances this iterator to the next sample. Returns true if there is one, false if not.
+				// This iterates over the entire track, ignoring timeline mapping.
+				ELECTRABASE_API bool Next();
+
+				// Recedes this iterator to the previous sample. Returns true if there is one, false if not.
+				// This iterates over the entire track, ignoring timeline mapping.
+				ELECTRABASE_API bool Prev();
+
+				// Returns whether the iterator points to the first overall sample, ignoring mapping to the timeline.
+				bool IsFirst() const
+				{ return SampleNumber == 0; }
+				// Returns whether the iterator points to the last overall sample, ignoring mapping to the timeline.
+				bool IsLast() const
+				{ return SampleNumber+1 >= Convs.NumTotalSamples; }
+
+				// Same as above, but obeying the timeline mapping and taking into consideration
+				// any required earlier sync frame and later frames due to reordering.
+				ELECTRABASE_API bool NextEffective();
+				ELECTRABASE_API bool PrevEffective();
+				bool IsFirstEffective() const
+				{ return SampleNumber <= Convs.FirstSample.SyncSampleNumber; }
+				bool IsLastEffective() const
+				{ return SampleNumber+1 >= Convs.LastSample.LastSampleNumber; }
+
+				// Creates a copy of this iterator.
+				TSharedPtr<FIterator> Clone()
+				{ return MakeShareable(new FIterator(*this)); }
+			protected:
+				FIterator() = default;
+				FIterator(const FIterator& InOther) = default;
+				friend class FMP4Track;
+
+				void Update();
+
+				TWeakPtr<FMP4Track> Track;
+				FConvenience Convs;
+				FSTSZBoxIterator stszIt;
+				FSTTSBoxIterator sttsIt;
+				FCTTSBoxIterator cttsIt;
+				FSTSCBoxIterator stscIt;
+				FSTSSBoxIterator stssIt;
+				FSTCOBoxIterator stcoIt;
+				FSBGPBoxIterator rapIt;
+				uint32 SampleNumber = 0;
+
+				FTimeFraction CurrentDTS;
+				FTimeFraction CurrentPTS;
+				FTimeFraction CurrentEffectiveDTS;
+				FTimeFraction CurrentEffectivePTS;
+				FTimeFraction CurrentDuration;
+				int64 CurrentSampleFileOffset = 0;
+				int64 CurrentSampleSize = 0;
+				bool bCurrentIsSyncOrRAP = false;
+			};
+
+			// Create an interator starting at the first sample.
+			ELECTRABASE_API TSharedPtr<FIterator> CreateIterator();
+			// Create an interator starting at the last sample (used when iterating in reverse, crossing back from the beginning to the end)
+			ELECTRABASE_API TSharedPtr<FIterator> CreateIteratorAtLastFrame();
+			// Create an iterator starting at a keyframe on or before the given time, or at a later time within the given
+			// threshold should one be right after the given time and would not be selected due to timescale rounding issues.
+			ELECTRABASE_API TSharedPtr<FIterator> CreateIteratorAtKeyframe(FTimeValue InForTime, FTimeValue InLaterTimeThreshold=FTimeValue::GetZero());
+			// Create an interator starting at a given sample number.
+			ELECTRABASE_API TSharedPtr<FIterator> CreateIterator(uint32 InAtSampleNumber);
+
+			// Returns the number of samples in this track.
+			ELECTRABASE_API uint32 GetNumberOfSamples();
+
+			// Returns information about the first sample that is mapped to the 0-based timeline via `elst` box.
+			const FFirstSample& GetFirstSampleInfo() const
+			{ return Convs.FirstSample; }
+			// Returns information about the last sample that is mapped to the 0-based timeline via `elst` box.
+			const FLastSample& GetLastSampleInfo() const
+			{ return Convs.LastSample; }
+
+		protected:
+			FMP4Track(const TSharedPtr<FMP4BoxTRAK>& InTrakBox) : TrakBox(InTrakBox)
+			{ }
+
+			TSharedPtr<FMP4BoxTRAK> TrakBox;
+
+			TSharedPtr<FMP4BoxTKHD> TkhdBox;
+			TSharedPtr<FMP4BoxELST> ElstBox;
+			TSharedPtr<FMP4BoxMDHD> MdhdBox;
+			TSharedPtr<FMP4BoxSTTS> SttsBox;
+			TSharedPtr<FMP4BoxCTTS> CttsBox;
+			TSharedPtr<FMP4BoxSTSC> StscBox;
+			TSharedPtr<FMP4BoxSTSZ> StszBox;
+			TSharedPtr<FMP4BoxSTCO> StcoBox;
+			TSharedPtr<FMP4BoxSTSS> StssBox;
+			TSharedPtr<FMP4BoxUDTA> UdtaBox;
+
+			TArray<TSharedPtr<FMP4BoxSGPD>> SgpdBoxes;
+			TArray<TSharedPtr<FMP4BoxSBGP>> SbgpBoxes;
+
+			FConvenience Convs;
+
+			FMP4TrackMetadataCommon CommonMetadata;
+
+			bool bHasBeenPrepared = false;
+
+			FString LastErrorMessage;
+		};
+
+	} // namespace UtilitiesMP4
+
+} // namespace Electra
diff --git a/Engine/Plugins/Media/ElectraUtil/Source/ElectraBase/Public/Utilities/UtilitiesMP4.h b/Engine/Plugins/Media/ElectraUtil/Source/ElectraBase/Public/Utilities/UtilitiesMP4.h
new file mode 100644
--- /dev/null
+++ b/Engine/Plugins/Media/ElectraUtil/Source/ElectraBase/Public/Utilities/UtilitiesMP4.h
@@ -0,0 +1,194 @@
+// Copyright Epic Games, Inc. All Rights Reserved.
+
+#pragma once
+
+#include "CoreTypes.h"
+#include "Containers/Array.h"
+#include "Containers/UnrealString.h"
+#include "Templates/SharedPointer.h"
+#include "Templates/PimplPtr.h"
+#include "Math/NumericLimits.h"
+#include "IElectraBaseDataReader.h"
+
+namespace Electra
+{
+#if !PLATFORM_LITTLE_ENDIAN
+	static constexpr uint8 GetFromBigEndian(uint8 value)		{ return value; }
+	static constexpr int8 GetFromBigEndian(int8 value)			{ return value; }
+	static constexpr uint16 GetFromBigEndian(uint16 value)		{ return value; }
+	static constexpr int16 GetFromBigEndian(int16 value)		{ return value; }
+	static constexpr int32 GetFromBigEndian(int32 value)		{ return value; }
+	static constexpr uint32 GetFromBigEndian(uint32 value)		{ return value; }
+	static constexpr int64 GetFromBigEndian(int64 value)		{ return value; }
+	static constexpr uint64 GetFromBigEndian(uint64 value)		{ return value; }
+#else
+	static constexpr uint16 EndianSwap(uint16 value)			{ return (value >> 8) | (value << 8); }
+	static constexpr int16 EndianSwap(int16 value)				{ return int16(EndianSwap(uint16(value))); }
+	static constexpr uint32 EndianSwap(uint32 value)			{ return (value << 24) | ((value & 0xff00) << 8) | ((value >> 8) & 0xff00) | (value >> 24); }
+	static constexpr int32 EndianSwap(int32 value)				{ return int32(EndianSwap(uint32(value))); }
+	static constexpr uint64 EndianSwap(uint64 value)			{ return (uint64(EndianSwap(uint32(value & 0xffffffffU))) << 32) | uint64(EndianSwap(uint32(value >> 32))); }
+	static constexpr int64 EndianSwap(int64 value)				{ return int64(EndianSwap(uint64(value)));}
+	static constexpr uint8 GetFromBigEndian(uint8 value)		{ return value; }
+	static constexpr int8 GetFromBigEndian(int8 value)			{ return value; }
+	static constexpr uint16 GetFromBigEndian(uint16 value)		{ return EndianSwap(value); }
+	static constexpr int16 GetFromBigEndian(int16 value)		{ return EndianSwap(value); }
+	static constexpr int32 GetFromBigEndian(int32 value)		{ return EndianSwap(value); }
+	static constexpr uint32 GetFromBigEndian(uint32 value)		{ return EndianSwap(value); }
+	static constexpr int64 GetFromBigEndian(int64 value)		{ return EndianSwap(value); }
+	static constexpr uint64 GetFromBigEndian(uint64 value)		{ return EndianSwap(value); }
+#endif
+
+	namespace UtilitiesMP4
+	{
+		class FMP4BoxBase;
+
+		static constexpr uint32 MakeBoxAtom(const uint8 A, const uint8 B, const uint8 C, const uint8 D)
+		{
+			return (static_cast<uint32>(A) << 24) | (static_cast<uint32>(B) << 16) | (static_cast<uint32>(C) << 8) | static_cast<uint32>(D);
+		}
+
+		static constexpr uint32 MakeBoxAtom(const char A, const char B, const char C, const char D)
+		{
+			return (static_cast<uint32>(A) << 24) | (static_cast<uint32>(B) << 16) | (static_cast<uint32>(C) << 8) | static_cast<uint32>(D);
+		}
+
+		static FString GetPrintableBoxAtom(uint32 InAtom)
+		{
+			TCHAR tc[4];
+			tc[0] = (TCHAR) ((InAtom >> 24) & 255);
+			tc[1] = (TCHAR) ((InAtom >> 16) & 255);
+			tc[2] = (TCHAR) ((InAtom >>  8) & 255);
+			tc[3] = (TCHAR) ((InAtom >>  0) & 255);
+			for(int32 i=0;i<4; ++i)
+			{
+				tc[i] = tc[i] >= 32 && tc[i] <= 127 ? tc[i] : TCHAR(' ');
+			}
+			return FString::ConstructFromPtrSize(tc, 4);
+		}
+
+		static FString Printable4CC(const uint32 In4CC)
+		{
+			FString Out;
+			// Not so much just printable as alphanumeric.
+			for(uint32 i=0, Atom=In4CC; i<4; ++i, Atom<<=8)
+			{
+				int32 v = Atom >> 24;
+				if ((v >= 'A' && v <= 'Z') || (v >= 'a' && v <= 'z') || (v >= '0' && v <= '9') || v == '_' || v == '.')
+				{
+					Out.AppendChar(v);
+				}
+				else
+				{
+					// Not alphanumeric, return it as a hex string.
+					return FString::Printf(TEXT("%08x"), In4CC);
+				}
+			}
+			return Out;
+		}
+
+		struct FMP4BoxInfo
+		{
+			TConstArrayView<uint8> Data;
+			uint8 UUID[16] {};
+			int64 Size = 0;
+			int64 Offset = 0;
+			uint32 Type = 0;
+			uint32 DataOffset = 0;
+		#if !UE_BUILD_SHIPPING
+			char Name[5]{0};
+		#endif
+		};
+		struct FMP4BoxData : public FMP4BoxInfo
+		{
+			TArray<uint8> DataBuffer;
+		};
+
+
+		class ELECTRABASE_API FMP4AtomReaderHelper
+		{
+		public:
+			FMP4AtomReaderHelper(const TConstArrayView<uint8>& InData);
+
+			bool ParseIntoBoxInfo(FMP4BoxInfo& OutBoxInfo, int64 InAtFileOffset);
+
+			int64 GetCurrentOffset() const;
+			int64 GetNumBytesRemaining() const;
+			const uint8* GetCurrentDataPointer() const;
+			void SetCurrentOffset(int64 InNewOffset);
+
+			template <typename T>
+			bool Read(T& OutValue)
+			{
+				T Temp = 0;
+				int64 NumRead = ReadData(&Temp, sizeof(T));
+				if (NumRead == sizeof(T))
+				{
+					OutValue = ValueFromBigEndian(Temp);
+					return true;
+				}
+				return false;
+			}
+
+			bool ReadVersionAndFlags(uint8& OutVersion, uint32& OutFlags);
+			bool ReadString(FString& OutString, uint16 InNumBytes);
+			bool ReadStringUTF8(FString& OutString, int32 InNumBytes);
+			bool ReadStringUTF16(FString& OutString, int32 InNumBytes);
+			bool ReadBytes(void* OutBuffer, int32 InNumBytes);
+			bool ReadAsNumber(int64& OutValue, int32 InNumBytes);
+			bool ReadAsNumber(uint64& OutValue, int32 InNumBytes);
+			bool ReadAsNumber(float& OutValue);
+			bool ReadAsNumber(double& OutValue);
+			bool SkipBytes(int32 InNumBytes)
+			{
+				return ReadData(nullptr, InNumBytes) == InNumBytes;
+			}
+		private:
+			template <typename T>
+			T ValueFromBigEndian(const T value)
+			{ return GetFromBigEndian(value); }
+			int32 ReadData(void* IntoBuffer, int32 NumBytesToRead);
+
+			const uint8* DataPtr = nullptr;
+			int64 DataSize = 0;
+			int64 CurrentOffset = 0;
+		};
+
+
+		class ELECTRABASE_API FMP4BoxLocatorReader
+		{
+		public:
+			FMP4BoxLocatorReader() = default;
+			~FMP4BoxLocatorReader();
+
+			bool LocateAndReadRootBoxes(TArray<TSharedPtr<FMP4BoxData, ESPMode::ThreadSafe>>& OutBoxInfos, const TSharedPtr<IBaseDataReader, ESPMode::ThreadSafe>& InDataReader, const TArray<uint32>& InFirstBoxes, const TArray<uint32>& InStopAfterBoxes, const TArray<uint32>& InReadDataOfBoxes, IBaseDataReader::FCancellationCheckDelegate InCheckCancellationDelegate);
+			FString GetLastError() const
+			{ return LastError; }
+		private:
+			FString LastError;
+			int64 CurrentOffset = 0;
+		};
+
+
+		class ELECTRABASE_API FMP4BoxTreeParser
+		{
+		public:
+			bool ParseBoxTree(const TSharedPtr<FMP4BoxInfo, ESPMode::ThreadSafe>& InRootBox);
+			TSharedPtr<FMP4BoxBase> GetBoxTree() const
+			{ return BoxTree; }
+		private:
+			bool ParseBoxTreeInternal(const TWeakPtr<FMP4BoxBase>& InParent, const FMP4BoxInfo& InBox);
+			TSharedPtr<FMP4BoxBase> BoxTree;
+		};
+
+	} // namespace UtilitiesMP4
+
+
+	class IFileDataReader : public IBaseDataReader
+	{
+	public:
+		static ELECTRABASE_API TSharedPtr<IFileDataReader, ESPMode::ThreadSafe> Create();
+		virtual ~IFileDataReader() = default;
+		virtual bool Open(const FString& InFilename) = 0;
+	};
+
+} // namespace Electra
