ElectraBase: Enable recursive parsing of `wave` box

#rnx

diff --git a/Engine/Plugins/Media/ElectraUtil/Source/ElectraBase/Public/Utilities/MP4Boxes/MP4Boxes.h b/Engine/Plugins/Media/ElectraUtil/Source/ElectraBase/Public/Utilities/MP4Boxes/MP4Boxes.h
--- a/Engine/Plugins/Media/ElectraUtil/Source/ElectraBase/Public/Utilities/MP4Boxes/MP4Boxes.h
+++ b/Engine/Plugins/Media/ElectraUtil/Source/ElectraBase/Public/Utilities/MP4Boxes/MP4Boxes.h
@@ -1806,6 +1806,7 @@ namespace Electra
 			static TSharedPtr<FMP4BoxBase> Create(const TWeakPtr<FMP4BoxBase>& InParent, const FMP4BoxInfo& InBoxInfo)
 			{ return MakeShareable<FMP4BoxWAVE>(new FMP4BoxWAVE(InParent, InBoxInfo)); }
 			virtual ~FMP4BoxWAVE() = default;
+			bool IsLeafBox() const override { return false; }
 		protected:
 			FMP4BoxWAVE(const TWeakPtr<FMP4BoxBase>& InParent, const FMP4BoxInfo& InBoxInfo)
 				: FMP4BoxBase(InParent, InBoxInfo)
