ElectraUtil: Added FFrameRate type to the parameter dictionary

#rnx

diff --git a/Engine/Plugins/Media/ElectraUtil/Source/ElectraBase/Private/ParameterDictionary.cpp b/Engine/Plugins/Media/ElectraUtil/Source/ElectraBase/Private/ParameterDictionary.cpp
--- a/Engine/Plugins/Media/ElectraUtil/Source/ElectraBase/Private/ParameterDictionary.cpp
+++ b/Engine/Plugins/Media/ElectraUtil/Source/ElectraBase/Private/ParameterDictionary.cpp
@@ -73,6 +73,12 @@ FVariantValue::FVariantValue(const FTimecode& TimecodeValue)
 	Set(TimecodeValue);
 }
 
+FVariantValue::FVariantValue(const FFrameRate& FramerateValue)
+	: DataType(EDataType::TypeUninitialized)
+{
+	Set(FramerateValue);
+}
+
 FVariantValue::FVariantValue(void* PointerValue)
 	: DataType(EDataType::TypeUninitialized)
 {
@@ -130,6 +136,11 @@ void FVariantValue::CopyInternal(const FVariantValue& FromOther)
 			Set(FromOther.GetTimecode());
 			break;
 		}
+		case EDataType::TypeFramerateValue:
+		{
+			Set(FromOther.GetFramerate());
+			break;
+		}
 		case EDataType::TypeVoidPointer:
 		{
 			Set(FromOther.GetPointer());
@@ -196,6 +207,7 @@ FVariant FVariantValue::ToFVariant() const
 		}
 		// Types that can't be converted.
 		case EDataType::TypeTimecodeValue:
+		case EDataType::TypeFramerateValue:
 		case EDataType::TypeSharedPointer:
 		default:
 		{
@@ -232,6 +244,12 @@ void FVariantValue::Clear()
 			Timecode->~FTimecode();
 			break;
 		}
+		case EDataType::TypeFramerateValue:
+		{
+			FFrameRate* Framerate = reinterpret_cast<FFrameRate*>(&DataBuffer);
+			Framerate->~FFrameRate();
+			break;
+		}
 		case EDataType::TypeUninitialized:
 		case EDataType::TypeDouble:
 		case EDataType::TypeInt64:
@@ -325,6 +343,15 @@ FVariantValue& FVariantValue::Set(const FTimecode& TimecodeValue)
 	return *this;
 }
 
+FVariantValue& FVariantValue::Set(const FFrameRate& FramerateValue)
+{
+	Clear();
+	FFrameRate* ValuePtr = reinterpret_cast<FFrameRate*>(&DataBuffer);
+	*ValuePtr = FramerateValue;
+	DataType = EDataType::TypeFramerateValue;
+	return *this;
+}
+
 FVariantValue& FVariantValue::Set(void* PointerValue)
 {
 	Clear();
@@ -441,6 +468,20 @@ const FTimecode& FVariantValue::GetTimecode() const
 	}
 }
 
+const FFrameRate& FVariantValue::GetFramerate() const
+{
+	if (ensure(DataType == EDataType::TypeFramerateValue))
+	{
+		const FFrameRate* Framerate = reinterpret_cast<const FFrameRate*>(&DataBuffer);
+		return *Framerate;
+	}
+	else
+	{
+		static FFrameRate Empty;
+		return Empty;
+	}
+}
+
 void* const & FVariantValue::GetPointer() const
 {
 	if (ensure(DataType == EDataType::TypeVoidPointer))
@@ -540,6 +581,16 @@ FTimecode FVariantValue::SafeGetTimecode(const FTimecode& Default) const
 	return Default;
 }
 
+FFrameRate FVariantValue::SafeGetFramerate(const FFrameRate& Default) const
+{
+	if (DataType == EDataType::TypeFramerateValue)
+	{
+		const FFrameRate* Framerate = reinterpret_cast<const FFrameRate*>(&DataBuffer);
+		return *Framerate;
+	}
+	return Default;
+}
+
 
 void* FVariantValue::SafeGetPointer(void* Default) const
 {
diff --git a/Engine/Plugins/Media/ElectraUtil/Source/ElectraBase/Public/MediaDecoderOutput.h b/Engine/Plugins/Media/ElectraUtil/Source/ElectraBase/Public/MediaDecoderOutput.h
--- a/Engine/Plugins/Media/ElectraUtil/Source/ElectraBase/Public/MediaDecoderOutput.h
+++ b/Engine/Plugins/Media/ElectraUtil/Source/ElectraBase/Public/MediaDecoderOutput.h
@@ -357,4 +357,5 @@ static const FName FPSDenominator(TEXT("fps_denom"));
 static const FName PixelDataScale(TEXT("pix_datascale"));
 static const FName Timecode(TEXT("timecode"));
 static const FName TMCDTimecode(TEXT("tmcd_timecode"));
+static const FName TMCDFramerate(TEXT("tmcd_framerate"));
 }
diff --git a/Engine/Plugins/Media/ElectraUtil/Source/ElectraBase/Public/ParameterDictionary.h b/Engine/Plugins/Media/ElectraUtil/Source/ElectraBase/Public/ParameterDictionary.h
--- a/Engine/Plugins/Media/ElectraUtil/Source/ElectraBase/Public/ParameterDictionary.h
+++ b/Engine/Plugins/Media/ElectraUtil/Source/ElectraBase/Public/ParameterDictionary.h
@@ -10,6 +10,7 @@
 #include "Misc/Variant.h"
 #include "Misc/Timespan.h"
 #include "Misc/Timecode.h"
+#include "Misc/FrameRate.h"
 
 
 namespace Electra
@@ -54,6 +55,7 @@ public:
 	explicit FVariantValue(const FTimeValue& TimeValue);
 	explicit FVariantValue(const FTimespan& TimespanValue);
 	explicit FVariantValue(const FTimecode& TimecodeValue);
+	explicit FVariantValue(const FFrameRate& FramerateValue);
 	explicit FVariantValue(void* PointerValue);
 	template<typename T> explicit FVariantValue(const TSharedPtr<T, ESPMode::ThreadSafe>& PointerValue)
 	: DataType(EDataType::TypeUninitialized)
@@ -69,6 +71,7 @@ public:
 	FVariantValue& Set(const FTimeValue& TimeValue);
 	FVariantValue& Set(const FTimespan& TimespanValue);
 	FVariantValue& Set(const FTimecode& TimecodeValue);
+	FVariantValue& Set(const FFrameRate& FramerateValue);
 	FVariantValue& Set(void* PointerValue);
 	template <typename T> FVariantValue& Set(const TSharedPtr<T, ESPMode::ThreadSafe>& PointerValue)
 	{
@@ -87,6 +90,7 @@ public:
 	const FTimeValue& GetTimeValue() const;
 	const FTimespan& GetTimespan() const;
 	const FTimecode& GetTimecode() const;
+	const FFrameRate& GetFramerate() const;
 	void* const & GetPointer() const;
 	template<typename T> TSharedPtr<T, ESPMode::ThreadSafe> GetSharedPointer() const
 	{
@@ -110,6 +114,7 @@ public:
 	FTimeValue SafeGetTimeValue(const FTimeValue& Default=FTimeValue()) const;
 	FTimespan SafeGetTimespan(const FTimespan& Default=FTimespan()) const;
 	FTimecode SafeGetTimecode(const FTimecode& Default=FTimecode()) const;
+	FFrameRate SafeGetFramerate(const FFrameRate& Default=FFrameRate()) const;
 	void* SafeGetPointer(void* Default=nullptr) const;
 	const TArray<uint8>& SafeGetArray() const;
 
@@ -123,6 +128,7 @@ public:
 		TypeTimeValue,
 		TypeTimespanValue,
 		TypeTimecodeValue,
+		TypeFramerateValue,
 		TypeVoidPointer,
 		TypeSharedPointer,
 		TypeU8Array,
@@ -156,6 +162,7 @@ private:
 		uint8 MemSizeTimeValue[sizeof(FTimeValue)];
 		uint8 MemSizeTimespan[sizeof(FTimespan)];
 		uint8 MemSizeTimecode[sizeof(FTimecode)];
+		uint8 MemSizeFramerate[sizeof(FFrameRate)];
 		uint8 MemSizeSharedPtrValue[sizeof(TSharedPtrHolder<uint8>)];
 		uint8 MemSizeTArray[sizeof(TArray<uint8>)];
 	};
@@ -281,5 +288,3 @@ private:
 
 
 } // namespace Electra
-
-
diff --git a/Engine/Plugins/Media/ElectraUtil/Source/ElectraSamples/Private/ElectraTextureSampleBase.cpp b/Engine/Plugins/Media/ElectraUtil/Source/ElectraSamples/Private/ElectraTextureSampleBase.cpp
--- a/Engine/Plugins/Media/ElectraUtil/Source/ElectraSamples/Private/ElectraTextureSampleBase.cpp
+++ b/Engine/Plugins/Media/ElectraUtil/Source/ElectraSamples/Private/ElectraTextureSampleBase.cpp
@@ -52,13 +52,17 @@ void IElectraTextureSampleBase::Initialize(FVideoDecoderOutput* InVideoDecoderOu
 	Colorimetry = VideoDecoderOutput->GetColorimetry();
 	HDRInfo = VideoDecoderOutput->GetHDRInformation();
 	TSharedPtr<const IVideoDecoderTimecode, ESPMode::ThreadSafe> TimecodePtr = VideoDecoderOutput->GetTimecode();
-
 	if (TimecodePtr.IsValid())
 	{
 		// Store this in case this is needed again.
 		DecoderTimecode = TimecodePtr;
 		Timecode = CreateTimecodeFromMPEGDefinition(Framerate, TimecodePtr->GetMPEGDefinition());
 	}
+	else if (VideoDecoderOutput->GetDict().HaveKey(IDecoderOutputOptionNames::TMCDTimecode))
+	{
+		Timecode = VideoDecoderOutput->GetDict().GetValue(IDecoderOutputOptionNames::TMCDTimecode).GetTimecode();
+		Framerate = VideoDecoderOutput->GetDict().GetValue(IDecoderOutputOptionNames::TMCDFramerate).SafeGetFramerate();
+	}
 
 	// Get various basic MP4-style colorimetry values (we default to video range Rec709 SDR)
 	bool bFullRange = false;
