Electra: Reworked the decoder output pool to fix the issue of crashes when it got destroyed while conversion jobs were still in-flight on the render thread.

#rnx

diff --git a/Engine/Plugins/Media/ElectraPlayer/Source/ElectraPlayerPlugin/Private/ElectraPlayerPlugin.cpp b/Engine/Plugins/Media/ElectraPlayer/Source/ElectraPlayerPlugin/Private/ElectraPlayerPlugin.cpp
--- a/Engine/Plugins/Media/ElectraPlayer/Source/ElectraPlayerPlugin/Private/ElectraPlayerPlugin.cpp
+++ b/Engine/Plugins/Media/ElectraPlayer/Source/ElectraPlayerPlugin/Private/ElectraPlayerPlugin.cpp
@@ -477,12 +477,6 @@ bool FElectraPlayerPlugin::CanReceiveAudioSamples(int32 NumFrames)
 	return MediaSamples->CanReceiveAudioSamples(NumFrames);
 }
 
-void FElectraPlayerPlugin::PrepareForDecoderShutdown()
-{
-	OutputTexturePool->PrepareForDecoderShutdown();
-}
-
-
 FString FElectraPlayerPlugin::GetVideoAdapterName() const
 {
 	return GRHIAdapterName;
diff --git a/Engine/Plugins/Media/ElectraPlayer/Source/ElectraPlayerPlugin/Private/ElectraPlayerPlugin.h b/Engine/Plugins/Media/ElectraPlayer/Source/ElectraPlayerPlugin/Private/ElectraPlayerPlugin.h
--- a/Engine/Plugins/Media/ElectraPlayer/Source/ElectraPlayerPlugin/Private/ElectraPlayerPlugin.h
+++ b/Engine/Plugins/Media/ElectraPlayer/Source/ElectraPlayerPlugin/Private/ElectraPlayerPlugin.h
@@ -168,7 +168,6 @@ private:
 	void PresentMetadataSample(const IMetaDataDecoderOutputPtr& InMetadataSample) override;
 	bool CanReceiveVideoSamples(int32 NumFrames) override;
 	bool CanReceiveAudioSamples(int32 NumFrames) override;
-	void PrepareForDecoderShutdown() override;
 	FString GetVideoAdapterName() const override;
 	TSharedPtr<IElectraPlayerResourceDelegate, ESPMode::ThreadSafe> GetResourceDelegate() const override;
 
diff --git a/Engine/Plugins/Media/ElectraPlayer/Source/ElectraPlayerPlugin/Private/ElectraPlayerPluginModule.cpp b/Engine/Plugins/Media/ElectraPlayer/Source/ElectraPlayerPlugin/Private/ElectraPlayerPluginModule.cpp
--- a/Engine/Plugins/Media/ElectraPlayer/Source/ElectraPlayerPlugin/Private/ElectraPlayerPluginModule.cpp
+++ b/Engine/Plugins/Media/ElectraPlayer/Source/ElectraPlayerPlugin/Private/ElectraPlayerPluginModule.cpp
@@ -89,7 +89,6 @@ public:
 			void PresentMetadataSample(const IMetaDataDecoderOutputPtr& InMetadataSample) override { }
 			bool CanReceiveVideoSamples(int32 NumFrames) override { return false; }
 			bool CanReceiveAudioSamples(int32 NumFrames) override { return false; }
-			void PrepareForDecoderShutdown() override { }
 			FString GetVideoAdapterName() const override { return FString(); }
 			TSharedPtr<IElectraPlayerResourceDelegate, ESPMode::ThreadSafe> GetResourceDelegate() const override { return PlayerResourceDlg; }
 		private:
diff --git a/Engine/Plugins/Media/ElectraPlayer/Source/ElectraPlayerRuntime/Private/ElectraPlayer.cpp b/Engine/Plugins/Media/ElectraPlayer/Source/ElectraPlayerRuntime/Private/ElectraPlayer.cpp
--- a/Engine/Plugins/Media/ElectraPlayer/Source/ElectraPlayerRuntime/Private/ElectraPlayer.cpp
+++ b/Engine/Plugins/Media/ElectraPlayer/Source/ElectraPlayerRuntime/Private/ElectraPlayer.cpp
@@ -426,20 +426,12 @@ void FElectraPlayer::CloseInternal(bool bKillAfterClose)
 	if (Player->RendererVideo.IsValid())
 	{
 		Player->RendererVideo->DetachPlayer();
-		// In addition we prepare the video renderer for decoder shutdown now.
-		Player->RendererVideo->PrepareForDecoderShutdown();
 	}
 	if (Player->RendererAudio.IsValid())
 	{
 		Player->RendererAudio->DetachPlayer();
 	}
 
-	TSharedPtr<IElectraPlayerAdapterDelegate, ESPMode::ThreadSafe> PinnedAdapterDelegate = AdapterDelegate.Pin();
-	if (PinnedAdapterDelegate.IsValid())
-	{
-		PinnedAdapterDelegate->PrepareForDecoderShutdown();
-	}
-
 	// Next up we clear out the sample queues.
 	// NOTE that it is important we use the On..Flush() methods here and not simply clear out the queues.
 	// The Flush() methods do more than that that is required to do and we don't need to duplicate that here.
diff --git a/Engine/Plugins/Media/ElectraPlayer/Source/ElectraPlayerRuntime/Private/Runtime/Decoder/VideoDecoder.cpp b/Engine/Plugins/Media/ElectraPlayer/Source/ElectraPlayerRuntime/Private/Runtime/Decoder/VideoDecoder.cpp
--- a/Engine/Plugins/Media/ElectraPlayer/Source/ElectraPlayerRuntime/Private/Runtime/Decoder/VideoDecoder.cpp
+++ b/Engine/Plugins/Media/ElectraPlayer/Source/ElectraPlayerRuntime/Private/Runtime/Decoder/VideoDecoder.cpp
@@ -807,8 +807,6 @@ IElectraDecoder::EOutputStatus FVideoDecoderImpl::HandleOutput()
 				continue;
 			}
 
-			Renderer->TickOutputBufferPool();
-
 			// Check if the renderer can accept the output we want to send to it.
 			if (Renderer.IsValid() && !Renderer->CanReceiveOutputFrames(1))
 			{
diff --git a/Engine/Plugins/Media/ElectraPlayer/Source/ElectraPlayerRuntime/Private/Runtime/Player/AdaptiveStreamingPlayerRenderers.cpp b/Engine/Plugins/Media/ElectraPlayer/Source/ElectraPlayerRuntime/Private/Runtime/Player/AdaptiveStreamingPlayerRenderers.cpp
--- a/Engine/Plugins/Media/ElectraPlayer/Source/ElectraPlayerRuntime/Private/Runtime/Player/AdaptiveStreamingPlayerRenderers.cpp
+++ b/Engine/Plugins/Media/ElectraPlayer/Source/ElectraPlayerRuntime/Private/Runtime/Player/AdaptiveStreamingPlayerRenderers.cpp
@@ -118,7 +118,6 @@ namespace Electra
 		UEMediaError Flush(const FParamDict& InOptions) override;
 		void StartRendering(const FParamDict& InOptions) override;
 		void StopRendering(const FParamDict& InOptions) override;
-		void TickOutputBufferPool() override;
 
 		UEMediaError ReturnVideoBuffer(IBuffer* Buffer, bool bRender, FParamDict& InSampleProperties);
 		UEMediaError ReturnAudioBuffer(IBuffer* Buffer, bool bRender, FParamDict& InSampleProperties);
@@ -568,11 +567,6 @@ void FAdaptiveStreamingWrappedRenderer::StopRendering(const FParamDict& InOption
 	WrappedRenderer->StopRendering(InOptions);
 }
 
-void FAdaptiveStreamingWrappedRenderer::TickOutputBufferPool()
-{
-	WrappedRenderer->TickOutputBufferPool();
-}
-
 void FAdaptiveStreamingWrappedRenderer::ReturnAllPendingBuffers(bool bForFlush)
 {
 	Lock.Lock();
diff --git a/Engine/Plugins/Media/ElectraPlayer/Source/ElectraPlayerRuntime/Private/Runtime/Player/SimpleElectraAudioPlayer.cpp b/Engine/Plugins/Media/ElectraPlayer/Source/ElectraPlayerRuntime/Private/Runtime/Player/SimpleElectraAudioPlayer.cpp
--- a/Engine/Plugins/Media/ElectraPlayer/Source/ElectraPlayerRuntime/Private/Runtime/Player/SimpleElectraAudioPlayer.cpp
+++ b/Engine/Plugins/Media/ElectraPlayer/Source/ElectraPlayerRuntime/Private/Runtime/Player/SimpleElectraAudioPlayer.cpp
@@ -252,17 +252,17 @@ public:
 		}
 	}
 	void Resume() override
-	{ 
+	{
 		if (TSharedPtr<IAdaptiveStreamingPlayer, ESPMode::ThreadSafe> LockedPlayer = Player)
 		{
 			LockedPlayer->Resume();
 		}
 	}
 	void Stop() override
-	{ 
+	{
 		if (TSharedPtr<IAdaptiveStreamingPlayer, ESPMode::ThreadSafe> LockedPlayer = Player)
 		{
-			LockedPlayer->Stop(); 
+			LockedPlayer->Stop();
 		}
 	}
 	bool HasErrored() const override
@@ -270,7 +270,7 @@ public:
 	FString GetError() const override
 	{ return ErrorMessage; }
 	bool HaveMetadata() const override
-	{ 
+	{
 		if (TSharedPtr<IAdaptiveStreamingPlayer, ESPMode::ThreadSafe> LockedPlayer = Player)
 		{
 			return LockedPlayer->HaveMetadata();
@@ -279,7 +279,7 @@ public:
 	}
 	int32 GetBinaryMetadata(TSharedPtr<TArray<uint8>, ESPMode::ThreadSafe>& OutMetadata) const override;
 	FTimespan GetDuration() const override
-	{ 
+	{
 		if (TSharedPtr<IAdaptiveStreamingPlayer, ESPMode::ThreadSafe> LockedPlayer = Player)
 		{
 			return LockedPlayer->GetDuration().GetAsTimespan();
@@ -289,7 +289,7 @@ public:
 	FTimespan GetPlayPosition() const override
 	{ return NextPTS; }
 	bool HasEnded() const override
-	{ 
+	{
 		if (TSharedPtr<IAdaptiveStreamingPlayer, ESPMode::ThreadSafe> LockedPlayer = Player)
 		{
 			return LockedPlayer->HasEnded();
@@ -297,42 +297,42 @@ public:
 		return false;
 	}
 	bool IsBuffering() const override
-	{ 
+	{
 		if (TSharedPtr<IAdaptiveStreamingPlayer, ESPMode::ThreadSafe> LockedPlayer = Player)
 		{
-			return LockedPlayer->IsBuffering(); 
+			return LockedPlayer->IsBuffering();
 		}
 		return false;
 	}
 	bool IsSeeking() const override
-	{ 
+	{
 		if (TSharedPtr<IAdaptiveStreamingPlayer, ESPMode::ThreadSafe> LockedPlayer = Player)
 		{
-			return LockedPlayer->IsSeeking(); 
+			return LockedPlayer->IsSeeking();
 		}
 		return false;
 	}
 	bool IsPlaying() const override
-	{ 
+	{
 		if (TSharedPtr<IAdaptiveStreamingPlayer, ESPMode::ThreadSafe> LockedPlayer = Player)
 		{
-			return LockedPlayer->IsPlaying(); 
+			return LockedPlayer->IsPlaying();
 		}
 		return false;
 	}
 	bool IsPaused() const override
-	{ 
+	{
 		if (TSharedPtr<IAdaptiveStreamingPlayer, ESPMode::ThreadSafe> LockedPlayer = Player)
 		{
-			return LockedPlayer->IsPaused(); 
+			return LockedPlayer->IsPaused();
 		}
 		return false;
 	}
 	void SetBitrateCeiling(int32 HighestSelectableBitrate) override
-	{ 
+	{
 		if (TSharedPtr<IAdaptiveStreamingPlayer, ESPMode::ThreadSafe> LockedPlayer = Player)
 		{
-			LockedPlayer->SetBitrateCeiling(HighestSelectableBitrate); 
+			LockedPlayer->SetBitrateCeiling(HighestSelectableBitrate);
 		}
 	}
 
@@ -772,7 +772,7 @@ private:
 			Analytics_NumSegmentsLoadedAtQuality = MakeAnalyticsStringFromMap(NumSegmentsLoadedAtQuality);
 			Analytics_CacheHitsAtQuality = MakeAnalyticsStringFromMap(CacheHitsAtQuality);
 		}
-	
+
 		FString MakeAnalyticsStringFromMap(TMap<int32, int32>& InMap) const
 		{
 			InMap.KeySort([](int32 A, int32 B){return A<B;});
@@ -830,6 +830,7 @@ public:
 	FSimpleElectraAudioPlayerRenderer(TSharedPtr<FSimpleElectraAudioPlayer, ESPMode::ThreadSafe> InParentPlayer)
 	{
 		Player = InParentPlayer;
+		DecoderOutputPool = FAudioPool::Create();
 	}
 	virtual ~FSimpleElectraAudioPlayerRenderer()
 	{ }
@@ -874,7 +875,8 @@ private:
 
 	TWeakPtr<FSimpleElectraAudioPlayer, ESPMode::ThreadSafe> Player;
 	TWeakPtr<IMediaRenderer, ESPMode::ThreadSafe> ParentRenderer;
-	TDecoderOutputObjectPool<IAudioDecoderOutput, FElectraPlayerPlatformAudioDecoderOutputFactory> DecoderOutputPool;
+	using FAudioPool = TDecoderOutputObjectPool<IAudioDecoderOutput, FElectraPlayerPlatformAudioDecoderOutputFactory>;
+	TSharedPtr<FAudioPool, ESPMode::ThreadSafe> DecoderOutputPool;
 	FParamDict BufferPoolProperties;
 	uint32 MaxBufferSize = 0;
 	int32 NumOutputAudioBuffersInUse = 0;
@@ -1794,7 +1796,7 @@ int64 FSimpleElectraAudioPlayer::GetNextSamples(FTimespan& OutPTS, bool& bOutIsF
 			lock.Unlock();
 			// Fill the output with silence.
 			FMemory::Memzero(OutBuffer, InNumFramesToGet * sizeof(int16) * nc);
-			return InNumFramesToGet; 
+			return InNumFramesToGet;
 		}
 	}
 	int32 NumFramesToGo = InNumFramesToGet < CurrentReadSampleBlock->NumFramesAvailable ? InNumFramesToGet : CurrentReadSampleBlock->NumFramesAvailable;
@@ -1933,7 +1935,7 @@ UEMediaError FSimpleElectraAudioPlayerRenderer::AcquireBuffer(IBuffer*& OutBuffe
 	{
 		return UEMEDIA_ERROR_INSUFFICIENT_DATA;
 	}
-	IAudioDecoderOutputPtr DecoderOutput = DecoderOutputPool.AcquireShared();
+	IAudioDecoderOutputPtr DecoderOutput = DecoderOutputPool->AcquireShared();
 	DecoderOutput->SetOwner(SharedThis(this));
 	FPlatformAtomics::InterlockedIncrement(&NumBuffersAcquiredForDecoder);
 	DecoderOutput->Reserve(MaxBufferSize);
diff --git a/Engine/Plugins/Media/ElectraPlayer/Source/ElectraPlayerRuntime/Private/Runtime/Renderer/RendererAudio.cpp b/Engine/Plugins/Media/ElectraPlayer/Source/ElectraPlayerRuntime/Private/Runtime/Renderer/RendererAudio.cpp
--- a/Engine/Plugins/Media/ElectraPlayer/Source/ElectraPlayerRuntime/Private/Runtime/Renderer/RendererAudio.cpp
+++ b/Engine/Plugins/Media/ElectraPlayer/Source/ElectraPlayerRuntime/Private/Runtime/Renderer/RendererAudio.cpp
@@ -69,7 +69,7 @@ const Electra::FParamDict& FElectraRendererAudio::GetBufferPoolProperties() cons
 UEMediaError FElectraRendererAudio::CreateBufferPool(const Electra::FParamDict& Parameters)
 {
 	// Create a new pool and overwrite all "pending" values...
-	DecoderOutputPool.Reset();
+	DecoderOutputPool = FAudioPool::Create();
 	NumOutputAudioBuffersInUse = 0;
 
 	const FVariantValue& variantNumBuffers = Parameters.GetValue(RenderOptionKeys::NumBuffers);
@@ -107,7 +107,7 @@ UEMediaError FElectraRendererAudio::AcquireBuffer(IBuffer*& OutBuffer, int32 Tim
 		return UEMEDIA_ERROR_INSUFFICIENT_DATA;
 	}
 
-	IAudioDecoderOutputPtr DecoderOutput = DecoderOutputPool.AcquireShared();
+	IAudioDecoderOutputPtr DecoderOutput = DecoderOutputPool->AcquireShared();
 	DecoderOutput->SetOwner(SharedThis(this));
 
 	FPlatformAtomics::InterlockedIncrement(&NumBuffersAcquiredForDecoder);
diff --git a/Engine/Plugins/Media/ElectraPlayer/Source/ElectraPlayerRuntime/Private/Runtime/Renderer/RendererAudio.h b/Engine/Plugins/Media/ElectraPlayer/Source/ElectraPlayerRuntime/Private/Runtime/Renderer/RendererAudio.h
--- a/Engine/Plugins/Media/ElectraPlayer/Source/ElectraPlayerRuntime/Private/Runtime/Renderer/RendererAudio.h
+++ b/Engine/Plugins/Media/ElectraPlayer/Source/ElectraPlayerRuntime/Private/Runtime/Renderer/RendererAudio.h
@@ -101,7 +101,8 @@ private:
 	TWeakPtr<FElectraPlayer, ESPMode::ThreadSafe> Player;
 	TWeakPtr<IMediaRenderer, ESPMode::ThreadSafe> ParentRenderer;
 
-	TDecoderOutputObjectPool<IAudioDecoderOutput, FElectraPlayerPlatformAudioDecoderOutputFactory> DecoderOutputPool;
+	using FAudioPool = TDecoderOutputObjectPool<IAudioDecoderOutput, FElectraPlayerPlatformAudioDecoderOutputFactory>;
+	TSharedPtr<FAudioPool, ESPMode::ThreadSafe> DecoderOutputPool;
 
 	uint32					 MaxBufferSize;
 	int32				     NumOutputAudioBuffersInUse;
@@ -135,5 +136,3 @@ private:
 	int32 NumBuffersAcquiredForDecoder;
 	TSharedPtr<Electra::IMediaRenderClock, ESPMode::ThreadSafe>	RenderClock;
 };
-
-
diff --git a/Engine/Plugins/Media/ElectraPlayer/Source/ElectraPlayerRuntime/Private/Runtime/Renderer/RendererVideo.cpp b/Engine/Plugins/Media/ElectraPlayer/Source/ElectraPlayerRuntime/Private/Runtime/Renderer/RendererVideo.cpp
--- a/Engine/Plugins/Media/ElectraPlayer/Source/ElectraPlayerRuntime/Private/Runtime/Renderer/RendererVideo.cpp
+++ b/Engine/Plugins/Media/ElectraPlayer/Source/ElectraPlayerRuntime/Private/Runtime/Renderer/RendererVideo.cpp
@@ -50,6 +50,7 @@ FElectraRendererVideo::FElectraRendererVideo(TSharedPtr<FElectraPlayer, ESPMode:
 	, NumBuffers(0)
 	, NumBuffersAcquiredForDecoder(0)
 {
+	DecoderOutputPool = FVideoPool::Create();
 }
 
 FElectraRendererVideo::~FElectraRendererVideo()
@@ -74,12 +75,6 @@ void FElectraRendererVideo::DetachPlayer()
 }
 
 
-void FElectraRendererVideo::PrepareForDecoderShutdown()
-{
-	DecoderOutputPool.PrepareForDecoderShutdown();
-}
-
-
 const Electra::FParamDict& FElectraRendererVideo::GetBufferPoolProperties() const
 {
 	return BufferPoolProperties;
@@ -94,7 +89,7 @@ const Electra::FParamDict& FElectraRendererVideo::GetBufferPoolProperties() cons
 void FElectraRendererVideo::AcquireFromPool(FVideoDecoderOutputPtr& DelayedImage)
 {
 	// Acquire a decoder output item
-	DelayedImage = DecoderOutputPool.AcquireShared();
+	DelayedImage = DecoderOutputPool->AcquireShared();
 	DelayedImage->SetOwner(SharedThis(this));
 
 	FPlatformAtomics::InterlockedIncrement(&NumOutputTexturesInUse);
@@ -113,7 +108,7 @@ UEMediaError FElectraRendererVideo::CreateBufferPool(const Electra::FParamDict&
 	// Update dict for later query via GetBufferPoolProperties
 	BufferPoolProperties.Set(RenderOptionKeys::MaxBuffers, Electra::FVariantValue((int64)RequestedNumBuffers));
 
-	// Currently, we only handle enlargement of the buffer pool. If the size should shrink,
+	// Currently, we only handle enlargement of the buffer pool.
 	if (RequestedNumBuffers != NumBuffers)
 	{
 		// Preallocate the buffer structures, but do not do the actual memory allocation!
@@ -285,11 +280,6 @@ void FElectraRendererVideo::StopRendering(const Electra::FParamDict& InOptions)
 	}
 }
 
-void FElectraRendererVideo::TickOutputBufferPool()
-{
-	DecoderOutputPool.Tick();
-}
-
 /**
  * Called from game thread to update buffer management
  */
diff --git a/Engine/Plugins/Media/ElectraPlayer/Source/ElectraPlayerRuntime/Private/Runtime/Renderer/RendererVideo.h b/Engine/Plugins/Media/ElectraPlayer/Source/ElectraPlayerRuntime/Private/Runtime/Renderer/RendererVideo.h
--- a/Engine/Plugins/Media/ElectraPlayer/Source/ElectraPlayerRuntime/Private/Runtime/Renderer/RendererVideo.h
+++ b/Engine/Plugins/Media/ElectraPlayer/Source/ElectraPlayerRuntime/Private/Runtime/Renderer/RendererVideo.h
@@ -49,8 +49,6 @@ public:
 	OpenError Open(const InstanceConfiguration& Config);
 	void Close(void);
 
-	void PrepareForDecoderShutdown();
-
 	//-------------------------------------------------------------------------
 	// Methods for UE interface
 	//
@@ -78,7 +76,7 @@ public:
 	UEMediaError ReleaseBufferPool() override;
 
 	bool CanReceiveOutputFrames(uint64 NumFrames) const override;
-	
+
 	bool GetEnqueuedFrameInfo(int32& OutNumberOfEnqueuedFrames, Electra::FTimeValue& OutDurationOfEnqueuedFrames) const override;
 
 	// Receives the render clock we need to update with the most recently rendered sample's timestamp.
@@ -99,19 +97,17 @@ public:
 	// Stops rendering of sample buffers
 	void StopRendering(const Electra::FParamDict& InOptions) override;
 
-	//! Tick any output buffer logic that might need regular updates
-	void TickOutputBufferPool() override;
-
 	//! Called from FElectraPlayerVideoDecoderOutput[Platform] when texture is returned to pool
 	void SampleReleasedToPool(IDecoderOutput* InDecoderOutput) override;
 
 private:
 
-
 	TWeakPtr<FElectraPlayer, ESPMode::ThreadSafe> Player;
 	TWeakPtr<IMediaRenderer, ESPMode::ThreadSafe> ParentRenderer;
 
-	TDecoderOutputObjectPool<FVideoDecoderOutput, FElectraPlayerPlatformVideoDecoderOutputFactory> DecoderOutputPool;
+	using FVideoPool = TDecoderOutputObjectPool<FVideoDecoderOutput, FElectraPlayerPlatformVideoDecoderOutputFactory>;
+	TSharedPtr<FVideoPool, ESPMode::ThreadSafe> DecoderOutputPool;
+
 	int32 NumOutputTexturesInUse;
 
 
@@ -137,15 +133,10 @@ private:
 		FVideoDecoderOutputPtr	DecoderOutput;
 	};
 
-
-
-
 	void AcquireFromPool(FVideoDecoderOutputPtr& DelayedImage);
 
-
 	Electra::FParamDict BufferPoolProperties;
 	int32 NumBuffers;
 	int32 NumBuffersAcquiredForDecoder;
 	TSharedPtr<Electra::IMediaRenderClock, ESPMode::ThreadSafe>	RenderClock;
 };
-
diff --git a/Engine/Plugins/Media/ElectraPlayer/Source/ElectraPlayerRuntime/Public/IElectraPlayerInterface.h b/Engine/Plugins/Media/ElectraPlayer/Source/ElectraPlayerRuntime/Public/IElectraPlayerInterface.h
--- a/Engine/Plugins/Media/ElectraPlayer/Source/ElectraPlayerRuntime/Public/IElectraPlayerInterface.h
+++ b/Engine/Plugins/Media/ElectraPlayer/Source/ElectraPlayerRuntime/Public/IElectraPlayerInterface.h
@@ -90,7 +90,6 @@ public:
 	virtual void PresentMetadataSample(const IMetaDataDecoderOutputPtr& InMetadataSample) = 0;
 	virtual bool CanReceiveVideoSamples(int32 NumFrames) = 0;
 	virtual bool CanReceiveAudioSamples(int32 NumFrames) = 0;
-	virtual void PrepareForDecoderShutdown() = 0;
 
 	virtual FString GetVideoAdapterName() const = 0;
 
diff --git a/Engine/Plugins/Media/ElectraUtil/Source/ElectraBase/ElectraBase.Build.cs b/Engine/Plugins/Media/ElectraUtil/Source/ElectraBase/ElectraBase.Build.cs
--- a/Engine/Plugins/Media/ElectraUtil/Source/ElectraBase/ElectraBase.Build.cs
+++ b/Engine/Plugins/Media/ElectraUtil/Source/ElectraBase/ElectraBase.Build.cs
@@ -11,7 +11,8 @@ namespace UnrealBuildTool.Rules
 			PublicDependencyModuleNames.AddRange(
 				new string[]
 				{
-					"Core"
+					"Core",
+                    "Engine"
 				});
 		}
 	}
diff --git a/Engine/Plugins/Media/ElectraUtil/Source/ElectraBase/Public/MediaDecoderOutput.h b/Engine/Plugins/Media/ElectraUtil/Source/ElectraBase/Public/MediaDecoderOutput.h
--- a/Engine/Plugins/Media/ElectraUtil/Source/ElectraBase/Public/MediaDecoderOutput.h
+++ b/Engine/Plugins/Media/ElectraUtil/Source/ElectraBase/Public/MediaDecoderOutput.h
@@ -3,11 +3,14 @@
 
 #include "Misc/Timespan.h"
 #include "Math/IntPoint.h"
+#include "Templates/Atomic.h"
 #include "Templates/SharedPointer.h"
 #include "Templates/RefCounting.h"
 #include "Containers/Queue.h"
 #include "Misc/ScopeLock.h"
 #include "HAL/PlatformProcess.h"
+#include "HAL/PlatformMisc.h"
+#include "Tickable.h"
 
 #include "ParameterDictionary.h"
 
@@ -36,276 +39,264 @@ public:
 template<typename ObjectType> class TElectraPoolDefaultObjectFactory
 {
 public:
-	static ObjectType* Create() { return new ObjectType; }
+	static ObjectType* Create()
+	{
+		return new ObjectType;
+	}
 };
 
+
+/**
+ * This class defines a pool of objects that inherit from the `IDecoderOutputPoolable` interface.
+ * The pool has no bound on the number of objects created, it merely tracks which objects have
+ * been handed out and awaits their return when the pool is destroyed.
+ */
 template<typename ObjectType, typename ObjectFactory = TElectraPoolDefaultObjectFactory<ObjectType>>
 class TDecoderOutputObjectPool
 {
-	static_assert(TPointerIsConvertibleFromTo<ObjectType, IDecoderOutputPoolable>::Value, "Poolable objects must implement the IDecoderOutputPoolable interface.");
+public:
+	/**
+	 * Creates a new object pool.
+	 * The pool has a custom deleter which only marks the pool as expired, but keeps it
+	 * around until all the elements it has ever handed out are returned and indicate
+	 * that they are no longer being used.
+	 */
+	static TSharedPtr<TDecoderOutputObjectPool, ESPMode::ThreadSafe> Create()
+	{
+		return MakeShareable(new TDecoderOutputObjectPool(), FPoolDeleter());
+	}
 
-	/** Object pool storage. */
-	class TStorage
+	/**
+	 * Acquires an object from the pool.
+	 * This returns either a new object or one that was used before and has been
+	 * returned to the pool. The object will not be re-initialized by default.
+	 * For this your managed object needs to implement the `InitializePoolable()`
+	 * method. Likewise, for object members that should be freed before the object
+	 * enters the pool for re-use, the object needs to implement the `ShutdownPoolable()`
+	 * method.
+	 */
+	TSharedRef<ObjectType, ESPMode::ThreadSafe> AcquireShared()
 	{
-		static float constexpr kStorageBusyWaitPollInterval = 0.002f;
+		return MakeShareable(ObjectPool->Acquire(), TObjectDeleter(ObjectPool));
+	}
+
+private:
+	static_assert(TPointerIsConvertibleFromTo<ObjectType, IDecoderOutputPoolable>::Value, "Poolable objects must implement the IDecoderOutputPoolable interface.");
 
+	/**
+	 * This class defines the actual pool.
+	 */
+	template<typename T=ObjectFactory>
+	class TObjectPool : public FTickableGameObject
+	{
 	public:
-		TStorage(ObjectFactory* InObjectFactoryInstance)
-			: ObjectFactoryInstance(InObjectFactoryInstance)
-		{}
+		TObjectPool() = default;
+		~TObjectPool() = default;
 
-		/** Destructor. */
-		~TStorage()
+		/** Called by the enclosing pool when it goes out of scope to let us handle cleanup of in-flight objects. */
+		void SetPendingDestruction(TSharedPtr<TObjectPool<T>, ESPMode::ThreadSafe> InSelf)
 		{
-			Reserve(0);
-
-			double StartTime = FPlatformTime::Seconds();
-			bool bDidWait = false;
-			ObjectType* Object;
-			while (WaitReadyForReuse.Dequeue(Object))
-			{
-				while (!Object->IsReadyForReuse() && GIsRunning)
-				{
-					// If we encounter an object that is still busy, we wait a bit and try again
-					// (Assumptions:
-					//  - there will not be that many busy objects
-					//  - usually the state of a busy object should become idle quite quickly)
-					bDidWait = true;
-					FPlatformProcess::Sleep(kStorageBusyWaitPollInterval);
-				}
-
-				Object->ShutdownPoolable();
-				delete Object;
-			}
-
-			if (bDidWait)
-			{
-				double EndTime = FPlatformTime::Seconds();
-				UE_LOG(LogTemp, Log, TEXT("[%p] TDecoderOutputObjectPool::TStorage::~TStorage() finished after %.3f msec!"), this, (EndTime - StartTime) * 1000.0);
-			}
+			Self = MoveTemp(InSelf);
+			bIsPendingDestruction = true;
 		}
 
-	public:
-
 		/** Acquire an object from the pool. */
 		ObjectType* Acquire()
 		{
-			ObjectType* Result = nullptr;
-			{
-				FScopeLock Lock(&CriticalSection);
+			// `Acquire()` cannot possibly be called when pool destruction is pending.
+			check(!bIsPendingDestruction);
 
-				if (Pool.Num() > 0)
-				{
-					Result = Pool.Pop(EAllowShrinking::No);
-				}
-				else
-				{
-					// Check for objects to ready to enter the pool & grab the first we find as our result...
-					// (we move all into the pool, we can to possibly safe on resources allocated by the objects)
-					ObjectType* PeekObject;
-					while (WaitReadyForReuse.Peek(PeekObject))
-					{
-						if (!PeekObject->IsReadyForReuse())
-						{
-							break;
-						}
-
-						WaitReadyForReuse.Pop();
-						PeekObject->ShutdownPoolable();
-						if (!Result)
-						{
-							Result = PeekObject;
-						}
-						else
-						{
-							Pool.Add(PeekObject);
-						}
-					}
-				}
-			}
+			ObjectType* Object = nullptr;
 
-			if (Result == nullptr)
+			CriticalSection.Lock();
+			// Handle objects that have just become reusable again.
+			HandleNewReturns();
+			// Take an object from the available list, if there is one.
+			if (Available.Num() > 0)
 			{
-				Result = ObjectFactoryInstance->Create();
+				Object = Available.Pop(EAllowShrinking::No);
 			}
+			CriticalSection.Unlock();
+			// If there is no object, we have to create a new one.
+			if (Object == nullptr)
+			{
+				Object = T::Create();
+			}
+			Object->InitializePoolable();
 
-			Result->InitializePoolable();
+			// Take note of the object we are handing out.
+			CriticalSection.Lock();
+			InFlight.Emplace(Object);
+			CriticalSection.Unlock();
 
-			return Result;
-		}
-
-		/** Get the number of objects stored. */
-		int32 Num() const
-		{
-			FScopeLock Lock(&CriticalSection);
-			return Pool.Num();
+			return Object;
 		}
 
 		/** Return the given object to the pool. */
-		void Release(ObjectType* Object)
+		void Release(ObjectType* InObject)
 		{
-			if (Object == nullptr)
+			if (InObject == nullptr)
 			{
 				return;
 			}
 
 			FScopeLock Lock(&CriticalSection);
+			check(InFlight.Contains(InObject));
 
-			if (Object->IsReadyForReuse())
+			InFlight.RemoveSwap(InObject, EAllowShrinking::No);
+			if (InObject->IsReadyForReuse())
 			{
-				Object->ShutdownPoolable();
-				Pool.Push(Object);
+				InObject->ShutdownPoolable();
+				Available.Push(InObject);
 			}
 			else
 			{
-				WaitReadyForReuse.Enqueue(Object);
+				Returned.Emplace(InObject);
 			}
 		}
 
-		/** Reserve the specified number of objects. */
-		void Reserve(uint32 NumObjects)
+		bool IsTickableWhenPaused() const override
+		{ return true; }
+		bool IsTickableInEditor() const override
+		{ return true; }
+		void Tick(float /*InDeltaTime*/) override
 		{
-			FScopeLock Lock(&CriticalSection);
-
-			while (NumObjects < (uint32)Pool.Num())
+			if (!bIsPendingDestruction)
 			{
-				delete Pool.Pop(EAllowShrinking::No);
+				return;
 			}
 
-			while (NumObjects > (uint32)Pool.Num())
+			CriticalSection.Lock();
+			// Handle objects that have just become reusable again.
+			HandleNewReturns();
+			CriticalSection.Unlock();
+			// Release all objects that are available now.
+			// We do not need to lock the mutex for this since `Acquire()` - which would need that -
+			// can no longer get called during a pending destruction as there is no user-code owner any more.
+			while(!Available.IsEmpty())
 			{
-				Pool.Push(ObjectFactoryInstance->Create());
+				delete Available.Pop();
+			}
+			// If there are no more in-flight objects or objects that await reusability we can destroy ourselved.
+			if (InFlight.IsEmpty() && Returned.IsEmpty())
+			{
+				// Bye-bye.
+				Self.Reset();
 			}
 		}
+		bool IsTickable() const override
+		{ return true; }
+		TStatId GetStatId() const override
+		{ RETURN_QUICK_DECLARE_CYCLE_STAT(TObjectPool, STATGROUP_Tickables); }
 
-		/** Regular tick call */
-		void Tick()
+	private:
+		void HandleNewReturns()
 		{
-			if (WaitReadyForReuse.IsEmpty())
+			// NOTE: This method expects the `CriticalSection` lock to have been locked already
+			//       in order to reduce the number of lock/unlocks.
+
+			// Move all the returned objects that have become ready for reuse into a temporary list.
+			TArray<ObjectType*> NewlyReusableObjects;
+			for(int32 i=Returned.Num()-1; i>=0; --i)
 			{
-				// Conservative early out to avoid CS: we will get any items missed the next time around
-				return;
+				if (Returned[i]->IsReadyForReuse())
+				{
+					NewlyReusableObjects.Emplace(Returned[i]);
+					Returned.RemoveAtSwap(i);
+				}
 			}
-
-			FScopeLock Lock(&CriticalSection);
-
-			ObjectType* Object;
-			while (WaitReadyForReuse.Peek(Object))
+			// If there are objects that have just become reusable we call their `ShutdownPoolable()`
+			// outside of our mutex lock in case what they are doing is costly.
+			if (NewlyReusableObjects.Num())
 			{
-				if (!Object->IsReadyForReuse())
+				CriticalSection.Unlock();
+				for(int32 i=NewlyReusableObjects.Num()-1; i>=0; --i)
 				{
-					break;
+					NewlyReusableObjects[i]->ShutdownPoolable();
 				}
-
-				Object->ShutdownPoolable();
-				Pool.Push(Object);
-				WaitReadyForReuse.Pop();
+				CriticalSection.Lock();
+				Available.Append(MoveTemp(NewlyReusableObjects));
 			}
 		}
 
-		void PrepareForDecoderShutdown()
-		{
-			// [...]
-		}
+		/** Critical section for synchronizing access to the free list. */
+		FCriticalSection CriticalSection;
 
-	private:
+		/** List of available objects. */
+		TArray<ObjectType*> Available;
 
-		/** Critical section for synchronizing access to the free list. */
-		mutable FCriticalSection CriticalSection;
+		/** List of in-flight objects. */
+		TArray<ObjectType*> InFlight;
 
-		/** List of unused objects. */
-		TArray<ObjectType*> Pool;
+		/** List of returned objects, waiting for reuseability. */
+		TArray<ObjectType*> Returned;
 
-		/** List of unused objects, waiting for reuse-ability. */
-		TQueue<ObjectType*> WaitReadyForReuse;
+		/** Pointer to self, which gets set when pool destruction is pending to lock it until complete. */
+		TSharedPtr<TObjectPool<T>, ESPMode::ThreadSafe> Self;
 
-		/** Object factory instance (nullptr by default) */
-		ObjectFactory* ObjectFactoryInstance;
+		/** Flag indicating whether or not the pool is pending destruction. */
+		TAtomic<bool> bIsPendingDestruction { false };
 	};
 
 
 	/** Deleter for pooled objects. */
-	class TDeleter
+	class TObjectDeleter
 	{
 	public:
 
 		/** Create and initialize a new instance. */
-		TDeleter(const TSharedRef<TStorage, ESPMode::ThreadSafe>& InStorage)
-			: StoragePtr(InStorage)
+		TObjectDeleter(const TSharedPtr<TObjectPool<ObjectFactory>, ESPMode::ThreadSafe>& InOwningPool)
+			: OwningPool(InOwningPool)
 		{ }
 
 		/** Function operator to execute deleter. */
 		void operator()(ObjectType* ObjectToDelete)
 		{
-			TSharedPtr<TStorage, ESPMode::ThreadSafe> PinnedStorage = StoragePtr.Pin();
-
-			if (PinnedStorage.IsValid())
+			TSharedPtr<TObjectPool<ObjectFactory>, ESPMode::ThreadSafe> PinnedStorage = OwningPool.Pin();
+			// The weak pointer must be lockable since the pool has a circular reference to itself
+			// during pool destruction.
+			if (ensure(PinnedStorage.IsValid()))
 			{
 				PinnedStorage->Release(ObjectToDelete);
 			}
 			else
 			{
-				delete ObjectToDelete;
+				if (ensure(ObjectToDelete->IsReadyForReuse()))
+				{
+					delete ObjectToDelete;
+				}
 			}
 		}
 
 	private:
-
-		/** Weak pointer to object pool storage. */
-		TWeakPtr<TStorage, ESPMode::ThreadSafe> StoragePtr;
+		/** Weak pointer to the owning object pool. */
+		TWeakPtr<TObjectPool<ObjectFactory>, ESPMode::ThreadSafe> OwningPool;
 	};
 
-public:
-
-	TDecoderOutputObjectPool(ObjectFactory* ObjectFactoryInstance = nullptr)
-		: Storage(MakeShareable(new TStorage(ObjectFactoryInstance)))
-	{ }
-
-	TDecoderOutputObjectPool(uint32 NumReserve, ObjectFactory* ObjectFactoryInstance = nullptr)
-		: Storage(MakeShareable(new TStorage(ObjectFactoryInstance)))
-	{
-		Storage->Reserve(NumReserve);
-	}
-
-public:
-
-	TSharedRef<ObjectType, ESPMode::ThreadSafe> AcquireShared()
-	{
-		ObjectType* Object = Storage->Acquire();
-		check(Object != nullptr);
-
-		return MakeShareable(Object, TDeleter(Storage));
-	}
-
-	int32 Num() const
-	{
-		return Storage->Num();
-	}
-
-	void Release(ObjectType* Object)
-	{
-		Storage->Release(Object);
-	}
 
-	void Reset(uint32 NumObjects = 0)
+	/**
+	 * This deleter class is invoked when the pool is to be destroyed.
+	 * We flag the actual pool as such and let its handling take care of this.
+	 */
+	class FPoolDeleter
 	{
-		Storage->Reserve(NumObjects);
-	}
-
-	void Tick()
-	{
-		Storage->Tick();
-	}
-
-	void PrepareForDecoderShutdown()
-	{
-		Storage->PrepareForDecoderShutdown();
-	}
+	public:
+		void operator()(TDecoderOutputObjectPool* InPool)
+		{
+			// Only mark the actual pool for destruction.
+			// It handles itself during `Tick()` and will destroy itself
+			// once all objects are returned it it.
+			check(InPool->ObjectPool.IsValid());
+			InPool->ObjectPool->SetPendingDestruction(MoveTemp(InPool->ObjectPool));
+		}
+	};
 
-private:
+	/** Private constructor. */
+	TDecoderOutputObjectPool() : ObjectPool(MakeShareable(new TObjectPool<ObjectFactory>()))
+	{ }
 
-	/** Storage for pooled objects. */
-	TSharedRef<TStorage, ESPMode::ThreadSafe> Storage;
+	/** The actual pool object. */
+	TSharedPtr<TObjectPool<ObjectFactory>, ESPMode::ThreadSafe> ObjectPool;
 };
 
 
diff --git a/Engine/Plugins/Media/ElectraUtil/Source/ElectraBase/Public/MediaRendererBase.h b/Engine/Plugins/Media/ElectraUtil/Source/ElectraBase/Public/MediaRendererBase.h
--- a/Engine/Plugins/Media/ElectraUtil/Source/ElectraBase/Public/MediaRendererBase.h
+++ b/Engine/Plugins/Media/ElectraUtil/Source/ElectraBase/Public/MediaRendererBase.h
@@ -215,13 +215,6 @@ namespace Electra
 		 * @param InOptions Reserved for future use.
 		 */
 		virtual void StopRendering(const FParamDict& InOptions) = 0;
-
-		/**
-		 * Work any regular tasks on the output buffer pool
-		 *
-		 * @note Expected to be called as needed from (normally) the decoder worker thread
-		 */
-		virtual void TickOutputBufferPool() {};
 	};
 
 
diff --git a/Engine/Plugins/Media/ElectraUtil/Source/ElectraSamples/Public/Android/ElectraTextureSample.h b/Engine/Plugins/Media/ElectraUtil/Source/ElectraSamples/Public/Android/ElectraTextureSample.h
--- a/Engine/Plugins/Media/ElectraUtil/Source/ElectraSamples/Public/Android/ElectraTextureSample.h
+++ b/Engine/Plugins/Media/ElectraUtil/Source/ElectraSamples/Public/Android/ElectraTextureSample.h
@@ -112,7 +112,6 @@ using FElectraTextureSampleRef  = TSharedRef<FElectraTextureSample, ESPMode::Thr
 
 class ELECTRASAMPLES_API FElectraTextureSamplePool : public TMediaObjectPool<FElectraTextureSample, FElectraTextureSamplePool>
 {
-	using ParentClass = TMediaObjectPool<FElectraTextureSample, FElectraTextureSamplePool>;
 	using TextureSample = FElectraTextureSample;
 
 public:
@@ -123,14 +122,8 @@ public:
 		return new TextureSample(Support);
 	}
 
-	void PrepareForDecoderShutdown()
-	{
-	}
-
 	void* GetCodecSurface();
 
 private:
 	TSharedPtr<FElectraTextureSampleSupport, ESPMode::ThreadSafe> Support;
 };
-
-
diff --git a/Engine/Plugins/Media/ElectraUtil/Source/ElectraSamples/Public/Apple/ElectraTextureSample.h b/Engine/Plugins/Media/ElectraUtil/Source/ElectraSamples/Public/Apple/ElectraTextureSample.h
--- a/Engine/Plugins/Media/ElectraUtil/Source/ElectraSamples/Public/Apple/ElectraTextureSample.h
+++ b/Engine/Plugins/Media/ElectraUtil/Source/ElectraSamples/Public/Apple/ElectraTextureSample.h
@@ -89,9 +89,7 @@ private:
 
 class ELECTRASAMPLES_API FElectraTextureSamplePool : public TMediaObjectPool<FElectraTextureSample, FElectraTextureSamplePool>
 {
-	using ParentClass = TMediaObjectPool<FElectraTextureSample, FElectraTextureSamplePool>;
 	using TextureSample = FElectraTextureSample;
-
 public:
 	FElectraTextureSamplePool()
 		: TMediaObjectPool<TextureSample, FElectraTextureSamplePool>(this)
@@ -103,11 +101,6 @@ public:
 		return new TextureSample(TexConv);
 	}
 
-	void PrepareForDecoderShutdown()
-	{
-	}
-
 private:
 	TSharedPtr<FElectraMediaTexConvApple, ESPMode::ThreadSafe> TexConv;
 };
-
diff --git a/Engine/Plugins/Media/ElectraUtil/Source/ElectraSamples/Public/Linux/ElectraTextureSample.h b/Engine/Plugins/Media/ElectraUtil/Source/ElectraSamples/Public/Linux/ElectraTextureSample.h
--- a/Engine/Plugins/Media/ElectraUtil/Source/ElectraSamples/Public/Linux/ElectraTextureSample.h
+++ b/Engine/Plugins/Media/ElectraUtil/Source/ElectraSamples/Public/Linux/ElectraTextureSample.h
@@ -57,9 +57,4 @@ private:
 using FElectraTextureSamplePtr = TSharedPtr<FElectraTextureSampleLinux, ESPMode::ThreadSafe>;
 using FElectraTextureSampleRef = TSharedRef<FElectraTextureSampleLinux, ESPMode::ThreadSafe>;
 using FElectraTextureSample = FElectraTextureSampleLinux;
-
-class FElectraTextureSamplePool : public TMediaObjectPool<FElectraTextureSampleLinux>
-{
-public:
-	void PrepareForDecoderShutdown() {}
-};
+using FElectraTextureSamplePool = TMediaObjectPool<FElectraTextureSampleLinux>;
diff --git a/Engine/Plugins/Media/ElectraUtil/Source/ElectraSamples/Public/Windows/ElectraTextureSample.h b/Engine/Plugins/Media/ElectraUtil/Source/ElectraSamples/Public/Windows/ElectraTextureSample.h
--- a/Engine/Plugins/Media/ElectraUtil/Source/ElectraSamples/Public/Windows/ElectraTextureSample.h
+++ b/Engine/Plugins/Media/ElectraUtil/Source/ElectraSamples/Public/Windows/ElectraTextureSample.h
@@ -77,10 +77,4 @@ private:
 
 using FElectraTextureSamplePtr = TSharedPtr<FElectraTextureSample, ESPMode::ThreadSafe>;
 using FElectraTextureSampleRef = TSharedRef<FElectraTextureSample, ESPMode::ThreadSafe>;
-
-class FElectraTextureSamplePool : public TMediaObjectPool<FElectraTextureSample>
-{
-public:
-	void PrepareForDecoderShutdown() {}
-};
-
+using FElectraTextureSamplePool = TMediaObjectPool<FElectraTextureSample>;
