Electra: Changes for better QuickTime compatibility

#rnx

diff --git a/Engine/Plugins/Media/ElectraCodecs/Source/ElectraDecoders/Private/Utils/MPEG/ElectraUtilsMPEGAudio.cpp b/Engine/Plugins/Media/ElectraCodecs/Source/ElectraDecoders/Private/Utils/MPEG/ElectraUtilsMPEGAudio.cpp
--- a/Engine/Plugins/Media/ElectraCodecs/Source/ElectraDecoders/Private/Utils/MPEG/ElectraUtilsMPEGAudio.cpp
+++ b/Engine/Plugins/Media/ElectraCodecs/Source/ElectraDecoders/Private/Utils/MPEG/ElectraUtilsMPEGAudio.cpp
@@ -72,10 +72,8 @@ namespace ElectraDecodersUtil
 			StreamTypeID = static_cast<FStreamType>(BitReader.GetBits(6));
 			// Skip upstream flag
 			BitReader.SkipBits(1);
-			if (BitReader.GetBits(1) != 1)	// reserved, must be 1
-			{
-				return false;
-			}
+			// Reserved '1'
+			BitReader.SkipBits(1);	// this bit must be 1, but it is sometimes incorrectly set to 0, so we do not check it.
 			BufferSize = BitReader.GetBits(24);
 			MaxBitrate = BitReader.GetBits(32);
 			AvgBitrate = BitReader.GetBits(32);
diff --git a/Engine/Plugins/Media/ElectraPlayer/Source/ElectraPlayerRuntime/Private/Runtime/Demuxer/ParserISO14496-12.cpp b/Engine/Plugins/Media/ElectraPlayer/Source/ElectraPlayerRuntime/Private/Runtime/Demuxer/ParserISO14496-12.cpp
--- a/Engine/Plugins/Media/ElectraPlayer/Source/ElectraPlayerRuntime/Private/Runtime/Demuxer/ParserISO14496-12.cpp
+++ b/Engine/Plugins/Media/ElectraPlayer/Source/ElectraPlayerRuntime/Private/Runtime/Demuxer/ParserISO14496-12.cpp
@@ -3683,7 +3683,15 @@ namespace Electra
 			FEntry e;
 			e.SampleCount = Entries[EntryIndex].SampleCount;
 			// For box version 0 return the unsigned value, otherwise convert the unsigned to signed (first to 32, then to 64 bits)
-			e.SampleOffset = Version == 0 ? (int64)Entries[EntryIndex].SampleOffset : (int64)((int32)Entries[EntryIndex].SampleOffset);
+//			e.SampleOffset = Version == 0 ? (int64)Entries[EntryIndex].SampleOffset : (int64)((int32)Entries[EntryIndex].SampleOffset);
+		/*
+			Note: We treat this box unconditionally as if it were version 1 using signed composition offsets.
+			      This addresses the issue that under the QuickTime brand this box is always using signed integers
+				  even in version 0.
+				  It stands to reason that no video will actually require unsigned values >0x7fffffff even if using
+				  HNS as timescale for which 0x7fffffff would still give about 214 seconds...
+		*/
+			e.SampleOffset = (int64)((int32)Entries[EntryIndex].SampleOffset);
 			return e;
 		}
 
diff --git a/Engine/Plugins/Media/ElectraPlayer/Source/ElectraPlayerRuntime/Private/Runtime/Player/mp4/ManifestMP4.cpp b/Engine/Plugins/Media/ElectraPlayer/Source/ElectraPlayerRuntime/Private/Runtime/Player/mp4/ManifestMP4.cpp
--- a/Engine/Plugins/Media/ElectraPlayer/Source/ElectraPlayerRuntime/Private/Runtime/Player/mp4/ManifestMP4.cpp
+++ b/Engine/Plugins/Media/ElectraPlayer/Source/ElectraPlayerRuntime/Private/Runtime/Player/mp4/ManifestMP4.cpp
@@ -938,18 +938,16 @@ void FManifestMP4Internal::FTimelineAssetMP4::UpdatePlayRangeEndInfo(const FTime
 							bFirst = false;
 							TrackLocalTime = PlayRangeEndInfo.Time.GetAsTimebase(TrkIt->GetTimescale());
 						}
-						// Look at the DTS, not the PTS since a frame with a larger PTS may still be referenced
-						// by earlier frames (In case of video: B frame reordering).
-						if (TrkIt->GetDTS() >= TrackLocalTime)
+						if (TrkIt->GetDTS() >= TrackLocalTime && TrkIt->GetPTS() >= TrackLocalTime)
 						{
-							TrackEndBytePos = TrkIt->GetSampleFileOffset() + TrkIt->GetSampleSize();
-							if (TrackEndBytePos > PlayRangeEndInfo.TotalFileEndOffset)
-							{
-								PlayRangeEndInfo.TotalFileEndOffset = TrackEndBytePos;
-							}
 							break;
 						}
 					}
+					TrackEndBytePos = TrkIt->GetSampleFileOffset() + TrkIt->GetSampleSize();
+					if (TrackEndBytePos > PlayRangeEndInfo.TotalFileEndOffset)
+					{
+						PlayRangeEndInfo.TotalFileEndOffset = TrackEndBytePos;
+					}
 				}
 				PlayRangeEndInfo.FileOffsetsPerTrackID.Emplace(TrackEndBytePos);
 			}
diff --git a/Engine/Plugins/Media/ElectraPlayer/Source/ElectraPlayerRuntime/Private/Runtime/Utilities/UtilsMPEG.cpp b/Engine/Plugins/Media/ElectraPlayer/Source/ElectraPlayerRuntime/Private/Runtime/Utilities/UtilsMPEG.cpp
--- a/Engine/Plugins/Media/ElectraPlayer/Source/ElectraPlayerRuntime/Private/Runtime/Utilities/UtilsMPEG.cpp
+++ b/Engine/Plugins/Media/ElectraPlayer/Source/ElectraPlayerRuntime/Private/Runtime/Utilities/UtilsMPEG.cpp
@@ -109,10 +109,7 @@ namespace Electra
 			StreamTypeID = static_cast<FStreamType>(BitReader.GetBits(6));
 			// Skip upstream flag
 			BitReader.SkipBits(1);
-			if (BitReader.GetBits(1) != 1)	// reserved, must be 1
-			{
-				return false;
-			}
+			BitReader.SkipBits(1);	// this bit is reserved and must be 1, but it is sometimes incorrectly set to 0, so do not check for it being 1!
 			BufferSize = BitReader.GetBits(24);
 			MaxBitrate = BitReader.GetBits(32);
 			AvgBitrate = BitReader.GetBits(32);
diff --git a/Engine/Plugins/Media/ElectraUtil/Source/ElectraBase/Private/Utilities/MP4Boxes/MP4Boxes.cpp b/Engine/Plugins/Media/ElectraUtil/Source/ElectraBase/Private/Utilities/MP4Boxes/MP4Boxes.cpp
--- a/Engine/Plugins/Media/ElectraUtil/Source/ElectraBase/Private/Utilities/MP4Boxes/MP4Boxes.cpp
+++ b/Engine/Plugins/Media/ElectraUtil/Source/ElectraBase/Private/Utilities/MP4Boxes/MP4Boxes.cpp
@@ -1263,6 +1263,13 @@ namespace Electra
 			ar.ReadVersionAndFlags(Parsed->Version, Parsed->Flags);
 			uint32 entry_count;
 			ar.Read(entry_count);
+		/*
+			Note: We read this box unconditionally as if it were version 1 using signed composition offsets.
+			      This addresses the issue that under the QuickTime brand this box is always using signed integers
+				  even in version 0.
+				  It stands to reason that no video will actually require unsigned values >0x7fffffff even if using
+				  HNS as timescale for which 0x7fffffff would still give about 214 seconds...
+
 			if (Parsed->Version == 0)
 			{
 				Parsed->Entries.SetNumUninitialized(entry_count);
@@ -1276,6 +1283,7 @@ namespace Electra
 				}
 			}
 			else if (Parsed->Version == 1)
+		*/
 			{
 				Parsed->Entries.SetNumUninitialized(entry_count);
 				int32 Value32;
diff --git a/Engine/Plugins/Media/ElectraUtil/Source/ElectraSamples/Public/ElectraTextureSampleUtils.h b/Engine/Plugins/Media/ElectraUtil/Source/ElectraSamples/Public/ElectraTextureSampleUtils.h
--- a/Engine/Plugins/Media/ElectraUtil/Source/ElectraSamples/Public/ElectraTextureSampleUtils.h
+++ b/Engine/Plugins/Media/ElectraUtil/Source/ElectraSamples/Public/ElectraTextureSampleUtils.h
@@ -18,14 +18,22 @@ namespace ElectraColorimetryUtils
 		{
 			case 1:				// Rec709
 			case 2:				// unknown
+			{
 				return UE::Color::EColorSpace::sRGB;
+			}
 			case 6:				// Rec601
 			case 7:				// Rec601 (historical)
+			{
 				return UE::Color::EColorSpace::sRGB;
+			}
 			case 9:				// Rec2020
+			{
 				return UE::Color::EColorSpace::Rec2020;
+			}
 			default:
-				check(!"Unexpected MPEG color primaries value!");
+			{
+				break;
+			}
 		}
 		return UE::Color::EColorSpace::sRGB;
 	}
@@ -35,17 +43,27 @@ namespace ElectraColorimetryUtils
 		switch (InMatrixCoefficients)
 		{
 			case 0:				// ID (RGB)
+			{
 				return UE::Color::EColorSpace::None;
+			}
 			case 1:				// Rec709
 			case 2:				// unknown
+			{
 				return UE::Color::EColorSpace::sRGB;
+			}
 			case 6:				// Rec601
 			case 7:				// Rec601 (historical)
+			{
 				return UE::Color::EColorSpace::sRGB;
+			}
 			case 9:				// Rec2020
+			{
 				return UE::Color::EColorSpace::Rec2020;
+			}
 			default:
-				check(!"Unexpected MPEG color primaries value!");
+			{
+				break;
+			}
 		}
 		return UE::Color::EColorSpace::sRGB;
 	}
@@ -56,18 +74,32 @@ namespace ElectraColorimetryUtils
 		{
 			case 1:
 			case 6:
+			case 7:
 			case 14:
-			case 15:
+			case 15: // Rec. ITU-R BT.1886-0
+			{
+				return UE::Color::EEncoding::sRGB;	// using sRGB in place of HLG for now
+			}
+			case 18: // Rec. ITU-R BT.2100-2
+			{
+				return UE::Color::EEncoding::sRGB;	// using sRGB in place of HLG for now
+			}
 			case 2:				// unknown
+			{
 				return UE::Color::EEncoding::sRGB;
+			}
 			case 8:
+			{
 				return UE::Color::EEncoding::Linear;
+			}
 			case 16:
+			{
 				return UE::Color::EEncoding::ST2084;
-			case 18:
-				return UE::Color::EEncoding::sRGB;	// using sRGB in place of HLG for now
+			}
 			default:
-				check(!"*** Unexpected MPEG color transfer characteristics!");
+			{
+				break;
+			}
 		}
 		return UE::Color::EEncoding::sRGB;
 	}
