Electra: Removal of platform specific HTTP streaming modules that are no longer being used.

#rnx

diff --git a/Engine/Plugins/Media/ElectraUtil/Source/ElectraHTTPStream/ElectraHTTPStream.Build.cs b/Engine/Plugins/Media/ElectraUtil/Source/ElectraHTTPStream/ElectraHTTPStream.Build.cs
--- a/Engine/Plugins/Media/ElectraUtil/Source/ElectraHTTPStream/ElectraHTTPStream.Build.cs
+++ b/Engine/Plugins/Media/ElectraUtil/Source/ElectraHTTPStream/ElectraHTTPStream.Build.cs
@@ -4,85 +4,8 @@ using UnrealBuildTool;
 
 public class ElectraHTTPStream : ModuleRules
 {
-    // Returns true when this platform will use the HTTP module from Unreal Engine
-    protected virtual bool bPlatformUsesGenericUE
-    {
-		get
-		{
-			return true;
-		}
-    }
-
-	// Returns true when this platform uses WinHttp
-	protected virtual bool bPlatformUsesWinHttp
-	{
-		get
-		{
-			return Target.Platform.IsInGroup(UnrealPlatformGroup.Windows);
-		}
-	}
-
-	// Returns whether or not this platform uses libCurl
-	protected virtual bool bPlatformUsesLibCurl
-	{
-		get
-		{
-			return Target.IsInPlatformGroup(UnrealPlatformGroup.Unix) || Target.IsInPlatformGroup(UnrealPlatformGroup.Android);
-		}
-	}
-
-
-	// Returns true if this platform supports xCurl (from Microsoft Game Core)
-	protected virtual bool bPlatformUsesXCurl
-	{
-		get
-		{
-			return false;
-		}
-	}
-
-
-	// Returns true if the platform requires all connections to be secure https:// connections.
-	protected virtual bool bPlatformRequiresSecureConnections
-	{
-		get
-		{
-			return false;
-		}
-	}
-
-	// Returns whether or not sockets are required for whatever HTTP library is being used.
-	protected virtual bool bPlatformRequiresSockets
-	{
-		get
-		{
-			return bPlatformUsesLibCurl;
-		}
-	}
-
-	// Returns whether or not this platform needs OpenSSL in combination with whatever HTTP library it is using.
-	protected virtual bool bPlatformRequiresOpenSSL
-	{
-		get
-		{
-			// WinHttp does not need OpenSSL.
-			if (bPlatformUsesWinHttp)
-			{
-				return false;
-			}
-			// libCurl is compiled to rely on OpenSSL
-			if (bPlatformUsesLibCurl)
-			{
-				return true;
-			}
-			return false;
-		}
-	}
-
 	public ElectraHTTPStream(ReadOnlyTargetRules Target) : base(Target)
 	{
-		PublicDefinitions.Add("ELECTRA_HTTPSTREAM_PACKAGE=1");
-
 		PublicDependencyModuleNames.AddRange(
 			new string[] {
 				"Core"
@@ -90,93 +13,8 @@ public class ElectraHTTPStream : ModuleRules
 
 		PrivateDependencyModuleNames.AddRange(
 			new string[] {
-				"ElectraBase"
+				"ElectraBase",
+				"HTTP"
 			});
-
-        if (bPlatformUsesGenericUE)
-        {
-		    PublicDefinitions.Add("ELECTRA_HTTPSTREAM_GENERIC_UE=1");
-		    PublicDefinitions.Add("ELECTRA_HTTPSTREAM_WINHTTP=0");
-			PublicDefinitions.Add("ELECTRA_HTTPSTREAM_LIBCURL=0");
-			PublicDefinitions.Add("ELECTRA_HTTPSTREAM_XCURL=0");
-		    PublicDefinitions.Add("ELECTRA_HTTPSTREAM_APPLE=0");
-			PrivateDependencyModuleNames.AddRange(
-				new string[] {
-					"HTTP",
-				});
-        }
-        else
-        {
-		    PublicDefinitions.Add("ELECTRA_HTTPSTREAM_GENERIC_UE=0");
-
-		    // When Curl is used we need to make sure there is no conflict with its use in the engine's HTTP module
-		    // as far as global initialization is concerned.
-		    if (bPlatformUsesLibCurl || bPlatformUsesXCurl)
-		    {
-			    PrivateDependencyModuleNames.AddRange(
-				    new string[] {
-					    "HTTP",
-				    });
-		    }
-
-		    if (bPlatformRequiresSockets)
-		    {
-			    PrivateDependencyModuleNames.AddRange(
-				    new string[] {
-					    "Sockets",
-				    }
-			    );
-		    }
-
-		    if (bPlatformUsesWinHttp)
-		    {
-			    PublicDefinitions.Add("ELECTRA_HTTPSTREAM_WINHTTP=1");
-			    AddEngineThirdPartyPrivateStaticDependencies(Target, "WinHttp");
-		    }
-		    else
-		    {
-			    PublicDefinitions.Add("ELECTRA_HTTPSTREAM_WINHTTP=0");
-		    }
-
-		    if (bPlatformUsesLibCurl || bPlatformUsesXCurl)
-		    {
-			    if (bPlatformUsesLibCurl)
-			    {
-				    PublicDefinitions.Add("ELECTRA_HTTPSTREAM_LIBCURL=1");
-				    PublicDefinitions.Add("ELECTRA_HTTPSTREAM_XCURL=0");
-				    AddEngineThirdPartyPrivateStaticDependencies(Target, "libcurl");
-			    }
-			    else
-			    {
-				    PublicDefinitions.Add("ELECTRA_HTTPSTREAM_LIBCURL=1");
-				    PublicDefinitions.Add("ELECTRA_HTTPSTREAM_XCURL=1");
-				    PublicDependencyModuleNames.Add("XCurl");
-			    }
-		    }
-		    else
-		    {
-			    PublicDefinitions.Add("ELECTRA_HTTPSTREAM_LIBCURL=0");
-			    PublicDefinitions.Add("ELECTRA_HTTPSTREAM_XCURL=0");
-		    }
-
-		    PublicDefinitions.Add("ELECTRA_HTTPSTREAM_REQUIRES_SECURE_CONNECTIONS=" + (bPlatformRequiresSecureConnections ? "1" : "0"));
-		
-		    if (bPlatformRequiresOpenSSL)
-		    {
-			    PrivateDependencyModuleNames.Add("SSL");
-			    AddEngineThirdPartyPrivateStaticDependencies(Target, "OpenSSL");
-		    }
-		    else
-		    {
-			    PrivateDefinitions.Add("WITH_SSL=0");
-		    }
-
-		    // Apple
-		    if (Target.IsInPlatformGroup(UnrealPlatformGroup.Apple))
-		    {
-			    PublicDefinitions.Add("ELECTRA_HTTPSTREAM_APPLE=1");
-			    PublicFrameworks.Add("Security");
-		    }
-        }
 	}
 }
diff --git a/Engine/Plugins/Media/ElectraUtil/Source/ElectraHTTPStream/Private/Apple/AppleElectraHTTPStream.cpp b/Engine/Plugins/Media/ElectraUtil/Source/ElectraHTTPStream/Private/Apple/AppleElectraHTTPStream.cpp
deleted file mode 100644
diff --git a/Engine/Plugins/Media/ElectraUtil/Source/ElectraHTTPStream/Private/Apple/AppleElectraHTTPStream.h b/Engine/Plugins/Media/ElectraUtil/Source/ElectraHTTPStream/Private/Apple/AppleElectraHTTPStream.h
deleted file mode 100644
diff --git a/Engine/Plugins/Media/ElectraUtil/Source/ElectraHTTPStream/Private/Apple/AppleLogError.cpp b/Engine/Plugins/Media/ElectraUtil/Source/ElectraHTTPStream/Private/Apple/AppleLogError.cpp
deleted file mode 100644
diff --git a/Engine/Plugins/Media/ElectraUtil/Source/ElectraHTTPStream/Private/Apple/AppleLogError.h b/Engine/Plugins/Media/ElectraUtil/Source/ElectraHTTPStream/Private/Apple/AppleLogError.h
deleted file mode 100644
diff --git a/Engine/Plugins/Media/ElectraUtil/Source/ElectraHTTPStream/Private/Curl/CurlElectra.h b/Engine/Plugins/Media/ElectraUtil/Source/ElectraHTTPStream/Private/Curl/CurlElectra.h
deleted file mode 100644
diff --git a/Engine/Plugins/Media/ElectraUtil/Source/ElectraHTTPStream/Private/Curl/CurlElectraHTTPStream.cpp b/Engine/Plugins/Media/ElectraUtil/Source/ElectraHTTPStream/Private/Curl/CurlElectraHTTPStream.cpp
deleted file mode 100644
diff --git a/Engine/Plugins/Media/ElectraUtil/Source/ElectraHTTPStream/Private/Curl/CurlElectraHTTPStream.h b/Engine/Plugins/Media/ElectraUtil/Source/ElectraHTTPStream/Private/Curl/CurlElectraHTTPStream.h
deleted file mode 100644
diff --git a/Engine/Plugins/Media/ElectraUtil/Source/ElectraHTTPStream/Private/Curl/CurlElectraHTTPStreamConfig.h b/Engine/Plugins/Media/ElectraUtil/Source/ElectraHTTPStream/Private/Curl/CurlElectraHTTPStreamConfig.h
deleted file mode 100644
diff --git a/Engine/Plugins/Media/ElectraUtil/Source/ElectraHTTPStream/Private/Curl/CurlElectraHTTPStreamImpl.h b/Engine/Plugins/Media/ElectraUtil/Source/ElectraHTTPStream/Private/Curl/CurlElectraHTTPStreamImpl.h
deleted file mode 100644
diff --git a/Engine/Plugins/Media/ElectraUtil/Source/ElectraHTTPStream/Private/Curl/CurlElectraHTTPStreamRequest.h b/Engine/Plugins/Media/ElectraUtil/Source/ElectraHTTPStream/Private/Curl/CurlElectraHTTPStreamRequest.h
deleted file mode 100644
diff --git a/Engine/Plugins/Media/ElectraUtil/Source/ElectraHTTPStream/Private/Curl/CurlLogError.cpp b/Engine/Plugins/Media/ElectraUtil/Source/ElectraHTTPStream/Private/Curl/CurlLogError.cpp
deleted file mode 100644
diff --git a/Engine/Plugins/Media/ElectraUtil/Source/ElectraHTTPStream/Private/Curl/CurlLogError.h b/Engine/Plugins/Media/ElectraUtil/Source/ElectraHTTPStream/Private/Curl/CurlLogError.h
deleted file mode 100644
diff --git a/Engine/Plugins/Media/ElectraUtil/Source/ElectraHTTPStream/Private/ElectraHTTPStream.cpp b/Engine/Plugins/Media/ElectraUtil/Source/ElectraHTTPStream/Private/ElectraHTTPStream.cpp
--- a/Engine/Plugins/Media/ElectraUtil/Source/ElectraHTTPStream/Private/ElectraHTTPStream.cpp
+++ b/Engine/Plugins/Media/ElectraUtil/Source/ElectraHTTPStream/Private/ElectraHTTPStream.cpp
@@ -1,7 +1,7 @@
 // Copyright Epic Games, Inc. All Rights Reserved.
 
 #include "ElectraHTTPStream.h"
-#include "PlatformElectraHTTPStream.h"
+#include "Generic/GenericElectraHTTPStream.h"
 
 TSharedPtr<IElectraHTTPStream, ESPMode::ThreadSafe> IElectraHTTPStream::Create(const Electra::FParamDict& InOptions)
 {
diff --git a/Engine/Plugins/Media/ElectraUtil/Source/ElectraHTTPStream/Private/ElectraHTTPStreamModule.cpp b/Engine/Plugins/Media/ElectraUtil/Source/ElectraHTTPStream/Private/ElectraHTTPStreamModule.cpp
--- a/Engine/Plugins/Media/ElectraUtil/Source/ElectraHTTPStream/Private/ElectraHTTPStreamModule.cpp
+++ b/Engine/Plugins/Media/ElectraUtil/Source/ElectraHTTPStream/Private/ElectraHTTPStreamModule.cpp
@@ -6,7 +6,7 @@
 #include "Modules/ModuleManager.h"
 #include "IElectraHTTPStreamModule.h"
 
-#include "PlatformElectraHTTPStream.h"
+#include "Generic/GenericElectraHTTPStream.h"
 
 #define LOCTEXT_NAMESPACE "ElectraHTTPStreamModule"
 
diff --git a/Engine/Plugins/Media/ElectraUtil/Source/ElectraHTTPStream/Private/Generic/GenericElectraHTTPStream.cpp b/Engine/Plugins/Media/ElectraUtil/Source/ElectraHTTPStream/Private/Generic/GenericElectraHTTPStream.cpp
--- a/Engine/Plugins/Media/ElectraUtil/Source/ElectraHTTPStream/Private/Generic/GenericElectraHTTPStream.cpp
+++ b/Engine/Plugins/Media/ElectraUtil/Source/ElectraHTTPStream/Private/Generic/GenericElectraHTTPStream.cpp
@@ -2,8 +2,6 @@
 
 #include "Generic/GenericElectraHTTPStream.h"
 
-#if ELECTRA_HTTPSTREAM_GENERIC_UE
-
 #include "Http.h"
 
 #include "HAL/IConsoleManager.h"
@@ -864,5 +862,3 @@ TSharedPtr<IElectraHTTPStream, ESPMode::ThreadSafe> FPlatformElectraHTTPStreamGe
 	}
 	return New;
 }
-
-#endif // ELECTRA_HTTPSTREAM_GENERIC_UE
diff --git a/Engine/Plugins/Media/ElectraUtil/Source/ElectraHTTPStream/Private/Generic/GenericElectraHTTPStream.h b/Engine/Plugins/Media/ElectraUtil/Source/ElectraHTTPStream/Private/Generic/GenericElectraHTTPStream.h
--- a/Engine/Plugins/Media/ElectraUtil/Source/ElectraHTTPStream/Private/Generic/GenericElectraHTTPStream.h
+++ b/Engine/Plugins/Media/ElectraUtil/Source/ElectraHTTPStream/Private/Generic/GenericElectraHTTPStream.h
@@ -2,8 +2,6 @@
 
 #pragma once
 
-#if ELECTRA_HTTPSTREAM_GENERIC_UE
-
 #include "CoreMinimal.h"
 #include "Templates/SharedPointer.h"
 #include "ParameterDictionary.h"
@@ -19,5 +17,3 @@ public:
 };
 
 typedef FPlatformElectraHTTPStreamGeneric FPlatformElectraHTTPStream;
-
-#endif // ELECTRA_HTTPSTREAM_GENERIC_UE
diff --git a/Engine/Plugins/Media/ElectraUtil/Source/ElectraHTTPStream/Private/WinHttp/WinHttpElectraHTTPStream.cpp b/Engine/Plugins/Media/ElectraUtil/Source/ElectraHTTPStream/Private/WinHttp/WinHttpElectraHTTPStream.cpp
deleted file mode 100644
diff --git a/Engine/Plugins/Media/ElectraUtil/Source/ElectraHTTPStream/Private/WinHttp/WinHttpElectraHTTPStream.h b/Engine/Plugins/Media/ElectraUtil/Source/ElectraHTTPStream/Private/WinHttp/WinHttpElectraHTTPStream.h
deleted file mode 100644
diff --git a/Engine/Plugins/Media/ElectraUtil/Source/ElectraHTTPStream/Private/WinHttp/WinHttpLogError.cpp b/Engine/Plugins/Media/ElectraUtil/Source/ElectraHTTPStream/Private/WinHttp/WinHttpLogError.cpp
deleted file mode 100644
diff --git a/Engine/Plugins/Media/ElectraUtil/Source/ElectraHTTPStream/Private/WinHttp/WinHttpLogError.h b/Engine/Plugins/Media/ElectraUtil/Source/ElectraHTTPStream/Private/WinHttp/WinHttpLogError.h
deleted file mode 100644
diff --git a/Engine/Plugins/Media/ElectraUtil/Source/ElectraHTTPStream/Public/Android/AndroidPlatformElectraHTTPStream.h b/Engine/Plugins/Media/ElectraUtil/Source/ElectraHTTPStream/Public/Android/AndroidPlatformElectraHTTPStream.h
deleted file mode 100644
diff --git a/Engine/Plugins/Media/ElectraUtil/Source/ElectraHTTPStream/Public/IOS/IOSPlatformElectraHTTPStream.h b/Engine/Plugins/Media/ElectraUtil/Source/ElectraHTTPStream/Public/IOS/IOSPlatformElectraHTTPStream.h
deleted file mode 100644
diff --git a/Engine/Plugins/Media/ElectraUtil/Source/ElectraHTTPStream/Public/Linux/LinuxPlatformElectraHTTPStream.h b/Engine/Plugins/Media/ElectraUtil/Source/ElectraHTTPStream/Public/Linux/LinuxPlatformElectraHTTPStream.h
deleted file mode 100644
diff --git a/Engine/Plugins/Media/ElectraUtil/Source/ElectraHTTPStream/Public/Mac/MacPlatformElectraHTTPStream.h b/Engine/Plugins/Media/ElectraUtil/Source/ElectraHTTPStream/Public/Mac/MacPlatformElectraHTTPStream.h
deleted file mode 100644
diff --git a/Engine/Plugins/Media/ElectraUtil/Source/ElectraHTTPStream/Public/PlatformElectraHTTPStream.h b/Engine/Plugins/Media/ElectraUtil/Source/ElectraHTTPStream/Public/PlatformElectraHTTPStream.h
deleted file mode 100644
diff --git a/Engine/Plugins/Media/ElectraUtil/Source/ElectraHTTPStream/Public/Unix/UnixPlatformElectraHTTPStream.h b/Engine/Plugins/Media/ElectraUtil/Source/ElectraHTTPStream/Public/Unix/UnixPlatformElectraHTTPStream.h
deleted file mode 100644
diff --git a/Engine/Plugins/Media/ElectraUtil/Source/ElectraHTTPStream/Public/Windows/WindowsPlatformElectraHTTPStream.h b/Engine/Plugins/Media/ElectraUtil/Source/ElectraHTTPStream/Public/Windows/WindowsPlatformElectraHTTPStream.h
deleted file mode 100644
