Speculative fix for android gles crash when movie player is active.

diff --git a/Engine/Plugins/Media/ElectraUtil/Source/ElectraSamples/Private/ElectraTextureSampleAndroid.cpp b/Engine/Plugins/Media/ElectraUtil/Source/ElectraSamples/Private/ElectraTextureSampleAndroid.cpp
--- a/Engine/Plugins/Media/ElectraUtil/Source/ElectraSamples/Private/ElectraTextureSampleAndroid.cpp
+++ b/Engine/Plugins/Media/ElectraUtil/Source/ElectraSamples/Private/ElectraTextureSampleAndroid.cpp
@@ -1086,6 +1086,9 @@ bool FElectraTextureSample::ConvertCpuOutputPath(FRHICommandListImmediate& RHICm
 
 	RHICmdList.Transition(FRHITransitionInfo(InDstTexture, ERHIAccess::Unknown, ERHIAccess::RTV));
 	FRHIRenderPassInfo RPInfo(InDstTexture, ERenderTargetActions::DontLoad_Store);
+
+	FBufferRHIRef TempVB = CreateTempMediaVertexBuffer();
+
 	RHICmdList.BeginRenderPass(RPInfo, TEXT("AndroidProcessVideo"));
 
 	// Update viewport.
@@ -1120,7 +1123,7 @@ bool FElectraTextureSample::ConvertCpuOutputPath(FRHICommandListImmediate& RHICm
 	// Update shader uniform parameters.
 	SetShaderParametersLegacyPS(RHICmdList, PixelShader, InputTexture, OutputDim, GetEncodingType(), ColorSpaceMtx);
 
-	RHICmdList.SetStreamSource(0, CreateTempMediaVertexBuffer(), 0);
+	RHICmdList.SetStreamSource(0, TempVB, 0);
 
 	RHICmdList.DrawPrimitive(0, 2, 1);
 
