Protron: Fixing an issue with the frame loader sharing a track iterator with the decoder that updates while the frame loader loads frames.

#rnx

diff --git a/Engine/Plugins/Media/ElectraPlayer/Source/ElectraProtron/Private/Player/ElectraProtronPlayerDecoder.cpp b/Engine/Plugins/Media/ElectraPlayer/Source/ElectraProtron/Private/Player/ElectraProtronPlayerDecoder.cpp
--- a/Engine/Plugins/Media/ElectraPlayer/Source/ElectraProtron/Private/Player/ElectraProtronPlayerDecoder.cpp
+++ b/Engine/Plugins/Media/ElectraPlayer/Source/ElectraProtron/Private/Player/ElectraProtronPlayerDecoder.cpp
@@ -329,7 +329,7 @@ void FElectraProtronPlayer::FImpl::FDecoderThread::UpdateTrackIterator(const FTi
 		return;
 	}
 
-	TSharedPtr<Electra::UtilitiesMP4::FMP4Track::FIterator> TkIt =
+	TSharedPtr<Electra::UtilitiesMP4::FMP4Track::FIterator, ESPMode::ThreadSafe> TkIt =
 		TrackSampleBuffer->TrackAndCodecInfo->MP4Track->CreateIteratorAtKeyframe(Electra::FTimeValue().SetFromTimespan(InForTime),
 			Electra::FTimeValue().SetFromMilliseconds(TrackSampleBuffer->TrackAndCodecInfo->bIsKeyframeOnlyFormat ? 0 : Config.NextKeyframeThresholdMillis));
 	if (!TkIt.IsValid())
@@ -582,7 +582,7 @@ void FElectraProtronPlayer::FImpl::FDecoderThread::StepTrackIterator(bool bPerfo
 			}
 			if (bShouldLoop)
 			{
-				TSharedPtr<Electra::UtilitiesMP4::FMP4Track::FIterator> TkIt = TrackSampleBuffer->TrackAndCodecInfo->MP4Track->CreateIterator(TrackIterator->GetFirstSampleInfo().SyncSampleNumber);
+				TSharedPtr<Electra::UtilitiesMP4::FMP4Track::FIterator, ESPMode::ThreadSafe> TkIt = TrackSampleBuffer->TrackAndCodecInfo->MP4Track->CreateIterator(TrackIterator->GetFirstSampleInfo().SyncSampleNumber);
 				check(TkIt.IsValid());
 				TrackIterator = MoveTemp(TkIt);
 				++LoopIndex;
@@ -607,7 +607,7 @@ void FElectraProtronPlayer::FImpl::FDecoderThread::StepTrackIterator(bool bPerfo
 				// go to whatever the last keyframe is to decode up to the last frame from there.
 				if (TrackSampleBuffer.IsValid() && TrackSampleBuffer->TrackAndCodecInfo->bIsKeyframeOnlyFormat)
 				{
-					TSharedPtr<Electra::UtilitiesMP4::FMP4Track::FIterator> TkIt = TrackSampleBuffer->TrackAndCodecInfo->MP4Track->CreateIterator(TrackIterator->GetLastSampleInfo().LastSampleNumber);
+					TSharedPtr<Electra::UtilitiesMP4::FMP4Track::FIterator, ESPMode::ThreadSafe> TkIt = TrackSampleBuffer->TrackAndCodecInfo->MP4Track->CreateIterator(TrackIterator->GetLastSampleInfo().LastSampleNumber);
 					check(TkIt.IsValid());
 					TrackIterator = MoveTemp(TkIt);
 					--LoopIndex;
diff --git a/Engine/Plugins/Media/ElectraPlayer/Source/ElectraProtron/Private/Player/ElectraProtronPlayerImpl.cpp b/Engine/Plugins/Media/ElectraPlayer/Source/ElectraProtron/Private/Player/ElectraProtronPlayerImpl.cpp
--- a/Engine/Plugins/Media/ElectraPlayer/Source/ElectraProtron/Private/Player/ElectraProtronPlayerImpl.cpp
+++ b/Engine/Plugins/Media/ElectraPlayer/Source/ElectraProtron/Private/Player/ElectraProtronPlayerImpl.cpp
@@ -230,7 +230,7 @@ void FElectraProtronPlayer::FImpl::InternalOpen(const FString& InFilename)
 	// We need the `moov` box now to determine what is inside this file.
 	// Yes, we looked at the presence of it above already as a quick reject before parsing the boxes.
 	// Now we need it for real and we know that it needs to be there, so here we go
-	TSharedPtr<Electra::UtilitiesMP4::FMP4BoxMOOV> MoovBox = StaticCastSharedPtr<Electra::UtilitiesMP4::FMP4BoxMOOV>(
+	TSharedPtr<Electra::UtilitiesMP4::FMP4BoxMOOV, ESPMode::ThreadSafe> MoovBox = StaticCastSharedPtr<Electra::UtilitiesMP4::FMP4BoxMOOV>(
 															(ParsedRootBoxes.FindByPredicate([](const Electra::UtilitiesMP4::FMP4BoxTreeParser& In)
 															{
 																return In.GetBoxTree()->GetType() == Electra::UtilitiesMP4::MakeBoxAtom('m','o','o','v');
@@ -260,7 +260,7 @@ void FElectraProtronPlayer::FImpl::InternalOpen(const FString& InFilename)
 	}
 
 	// Get all the tracks.
-	TArray<TSharedPtr<Electra::UtilitiesMP4::FMP4BoxTRAK>> AllTracks;
+	TArray<TSharedPtr<Electra::UtilitiesMP4::FMP4BoxTRAK, ESPMode::ThreadSafe>> AllTracks;
 	MoovBox->GetAllBoxInstances(AllTracks, Electra::UtilitiesMP4::MakeBoxAtom('t','r','a','k'));
 	// Prepare the internal track structure and check which tracks we can use and which ones we cannot.
 	Tracks.SetNum(AllTracks.Num());
@@ -440,7 +440,7 @@ void FElectraProtronPlayer::FImpl::InternalOpen(const FString& InFilename)
 							// Either way we need to tag every referenced track.
 							for(auto& RefTkId : References[0].TrackIDs)
 							{
-								TSharedPtr<FTrackInfo>* ReferencedTrack = Tracks.FindByPredicate([RefTkId](const TSharedPtr<FTrackInfo>& e){ return e->TrackID == RefTkId; });
+								TSharedPtr<FTrackInfo, ESPMode::ThreadSafe>* ReferencedTrack = Tracks.FindByPredicate([RefTkId](const TSharedPtr<FTrackInfo, ESPMode::ThreadSafe>& e){ return e->TrackID == RefTkId; });
 								if (ReferencedTrack)
 								{
 									if (!Tracks[TrkNum]->ReferencedTimecodeTrack.IsValid())
@@ -468,7 +468,7 @@ void FElectraProtronPlayer::FImpl::InternalOpen(const FString& InFilename)
 			if (NumTimecodeTracks == 1)
 			{
 				// Which track is the timecode?
-				TSharedPtr<FTrackInfo>* TimecodeTrack = Tracks.FindByPredicate([](const TSharedPtr<FTrackInfo>& In){ return In->CodecInfo.Type == ElectraProtronUtils::FCodecInfo::EType::Timecode; });
+				TSharedPtr<FTrackInfo, ESPMode::ThreadSafe>* TimecodeTrack = Tracks.FindByPredicate([](const TSharedPtr<FTrackInfo, ESPMode::ThreadSafe>& In){ return In->CodecInfo.Type == ElectraProtronUtils::FCodecInfo::EType::Timecode; });
 				check(TimecodeTrack); // there has to be one, otherwise we would not even get here, but for safeties sake...
 				if (TimecodeTrack)
 				{
@@ -501,7 +501,7 @@ void FElectraProtronPlayer::FImpl::InternalOpen(const FString& InFilename)
 			//auto CTSZBox = UdtaBox->FindBoxRecursive<Electra::UtilitiesMP4::FMP4BoxBase>(Electra::UtilitiesMP4::MakeBoxAtom(0xa9U,'T','S','Z'), 0);
 			if (CTIMBox.IsValid() && CTSCBox.IsValid())
 			{
-				auto GetValue = [](const TSharedPtr<Electra::UtilitiesMP4::FMP4BoxBase>& InBox) -> FString
+				auto GetValue = [](const TSharedPtr<Electra::UtilitiesMP4::FMP4BoxBase, ESPMode::ThreadSafe>& InBox) -> FString
 				{
 					Electra::UtilitiesMP4::FMP4AtomReaderHelper timReader(InBox->GetBoxData());
 					FString String;
@@ -623,7 +623,7 @@ void FElectraProtronPlayer::FImpl::InternalOpen(const FString& InFilename)
 	}
 }
 
-void FElectraProtronPlayer::FImpl::GetTrackCodecInfo(ElectraProtronUtils::FCodecInfo& OutCodecInfo, const TSharedPtr<Electra::UtilitiesMP4::FMP4BoxTRAK>& InTrack, uint32 InTrackID)
+void FElectraProtronPlayer::FImpl::GetTrackCodecInfo(ElectraProtronUtils::FCodecInfo& OutCodecInfo, const TSharedPtr<Electra::UtilitiesMP4::FMP4BoxTRAK, ESPMode::ThreadSafe>& InTrack, uint32 InTrackID)
 {
 	// There needs to be an `stsd` box in this track. We do not try the expected path of `trak`->`mdia`->`minf`->`stbl`->`stsd` as
 	// this is not that much faster and if the file is somewhat ill-formed we may not find it if it's grouped under elsewhere.
@@ -635,7 +635,7 @@ void FElectraProtronPlayer::FImpl::GetTrackCodecInfo(ElectraProtronUtils::FCodec
 		return;
 	}
 
-	auto AddChildren = [&OutCodecInfo](const TSharedPtr<Electra::UtilitiesMP4::FMP4BoxBase>& InSample) -> void
+	auto AddChildren = [&OutCodecInfo](const TSharedPtr<Electra::UtilitiesMP4::FMP4BoxBase, ESPMode::ThreadSafe>& InSample) -> void
 	{
 		for(auto& Ch : InSample->GetChildren())
 		{
@@ -667,7 +667,7 @@ void FElectraProtronPlayer::FImpl::GetTrackCodecInfo(ElectraProtronUtils::FCodec
 	// Based on the sample type, get it and see if it is using codec we support.
 	if (SampleType == Electra::UtilitiesMP4::FMP4BoxSampleEntry::ESampleType::Video)
 	{
-		auto SetResoAndFPSFromBox = [InTrack](ElectraProtronUtils::FCodecInfo::FVideo& InOutVideo, const TSharedPtr<Electra::UtilitiesMP4::FMP4BoxVisualSampleEntry>& InVisual) -> void
+		auto SetResoAndFPSFromBox = [InTrack](ElectraProtronUtils::FCodecInfo::FVideo& InOutVideo, const TSharedPtr<Electra::UtilitiesMP4::FMP4BoxVisualSampleEntry, ESPMode::ThreadSafe>& InVisual) -> void
 		{
 			InOutVideo.Width = InVisual->GetWidth();
 			InOutVideo.Height = InVisual->GetHeight();
@@ -679,7 +679,7 @@ void FElectraProtronPlayer::FImpl::GetTrackCodecInfo(ElectraProtronUtils::FCodec
 			}
 		};
 
-		TArray<TSharedPtr<Electra::UtilitiesMP4::FMP4BoxVisualSampleEntry>> Visuals;
+		TArray<TSharedPtr<Electra::UtilitiesMP4::FMP4BoxVisualSampleEntry, ESPMode::ThreadSafe>> Visuals;
 		Stsd->GetSampleDescriptions(Visuals);
 		check(Visuals.Num() == 1);
 		switch(Visuals[0]->GetType())
@@ -791,7 +791,7 @@ void FElectraProtronPlayer::FImpl::GetTrackCodecInfo(ElectraProtronUtils::FCodec
 	}
 	else if (SampleType == Electra::UtilitiesMP4::FMP4BoxSampleEntry::ESampleType::Audio)
 	{
-		TArray<TSharedPtr<Electra::UtilitiesMP4::FMP4BoxAudioSampleEntry>> Audios;
+		TArray<TSharedPtr<Electra::UtilitiesMP4::FMP4BoxAudioSampleEntry, ESPMode::ThreadSafe>> Audios;
 		Stsd->GetSampleDescriptions(Audios);
 		check(Audios.Num() == 1);
 
@@ -866,7 +866,7 @@ void FElectraProtronPlayer::FImpl::GetTrackCodecInfo(ElectraProtronUtils::FCodec
 	}
 	else if (SampleType == Electra::UtilitiesMP4::FMP4BoxSampleEntry::ESampleType::QTFFTimecode)
 	{
-		TArray<TSharedPtr<Electra::UtilitiesMP4::FMP4BoxQTFFTimecodeSampleEntry>> Timecodes;
+		TArray<TSharedPtr<Electra::UtilitiesMP4::FMP4BoxQTFFTimecodeSampleEntry, ESPMode::ThreadSafe>> Timecodes;
 		Stsd->GetSampleDescriptions(Timecodes);
 		check(Timecodes.Num() == 1);
 		switch(Timecodes[0]->GetType())
@@ -895,7 +895,7 @@ void FElectraProtronPlayer::FImpl::GetTrackCodecInfo(ElectraProtronUtils::FCodec
 	{
 		if (!bIsEncrypted)
 		{
-			TArray<TSharedPtr<Electra::UtilitiesMP4::FMP4BoxSampleEntry>> Entries;
+			TArray<TSharedPtr<Electra::UtilitiesMP4::FMP4BoxSampleEntry, ESPMode::ThreadSafe>> Entries;
 			Stsd->GetSampleDescriptions(Entries);
 			check(Entries.Num() == 1);
 			UE_LOG(LogElectraProtron, Warning, TEXT("Track of sample type \"%s\" is not supported, ignoring track #%u."), *Electra::UtilitiesMP4::GetPrintableBoxAtom(Entries[0]->GetType()), InTrackID);
diff --git a/Engine/Plugins/Media/ElectraPlayer/Source/ElectraProtron/Private/Player/ElectraProtronPlayerImpl.h b/Engine/Plugins/Media/ElectraPlayer/Source/ElectraProtron/Private/Player/ElectraProtronPlayerImpl.h
--- a/Engine/Plugins/Media/ElectraPlayer/Source/ElectraProtron/Private/Player/ElectraProtronPlayerImpl.h
+++ b/Engine/Plugins/Media/ElectraPlayer/Source/ElectraProtron/Private/Player/ElectraProtronPlayerImpl.h
@@ -133,12 +133,12 @@ private:
 
 	struct FTrackInfo
 	{
-		TArray<TWeakPtr<FTrackInfo>> IsReferencedByTracks;
+		TArray<TWeakPtr<FTrackInfo, ESPMode::ThreadSafe>> IsReferencedByTracks;
 		ElectraProtronUtils::FCodecInfo CodecInfo;
 		FString HumanReadableCodecFormat;
-		TSharedPtr<Electra::UtilitiesMP4::FMP4BoxTRAK> TrackBox;
-		TSharedPtr<Electra::UtilitiesMP4::FMP4Track> MP4Track;
-		TWeakPtr<FTrackInfo> ReferencedTimecodeTrack;
+		TSharedPtr<Electra::UtilitiesMP4::FMP4BoxTRAK, ESPMode::ThreadSafe> TrackBox;
+		TSharedPtr<Electra::UtilitiesMP4::FMP4Track, ESPMode::ThreadSafe> MP4Track;
+		TWeakPtr<FTrackInfo, ESPMode::ThreadSafe> ReferencedTimecodeTrack;
 		uint32 TrackID = 0;
 		bool bIsUsable = false;
 		bool bIsKeyframeOnlyFormat = false;
@@ -181,18 +181,18 @@ private:
 		FCriticalSection Lock;
 		TRangeSet<uint32> SampleRanges;
 		TSortedMap<uint32, FMP4SamplePtr> SampleMap;
-		TSharedPtr<FTrackInfo> TrackAndCodecInfo;
+		TSharedPtr<FTrackInfo, ESPMode::ThreadSafe> TrackAndCodecInfo;
 		uint32 TrackID;
 	};
 	using FMP4TrackSampleBufferPtr = TSharedPtr<FMP4TrackSampleBuffer, ESPMode::ThreadSafe>;
 
-	using FTrackIterator = TSharedPtr<Electra::UtilitiesMP4::FMP4Track::FIterator>;
+	using FTrackIterator = TSharedPtr<Electra::UtilitiesMP4::FMP4Track::FIterator, ESPMode::ThreadSafe>;
 
 	void StartThread();
 	void SendWorkerThreadMessage(FWorkerThreadMessage&& InMessage);
 
 	void InternalOpen(const FString& InFilename);
-	void GetTrackCodecInfo(ElectraProtronUtils::FCodecInfo& OutCodecInfo, const TSharedPtr<Electra::UtilitiesMP4::FMP4BoxTRAK>& InTrack, uint32 InTrackID);
+	void GetTrackCodecInfo(ElectraProtronUtils::FCodecInfo& OutCodecInfo, const TSharedPtr<Electra::UtilitiesMP4::FMP4BoxTRAK, ESPMode::ThreadSafe>& InTrack, uint32 InTrackID);
 
 	void UpdateTrackLoader(int32 InCodecTypeIndex);
 
@@ -211,7 +211,7 @@ private:
 	bool bAbort = false;
 	TArray<Electra::UtilitiesMP4::FMP4BoxTreeParser> ParsedRootBoxes;
 
-	TArray<TSharedPtr<FTrackInfo>> Tracks;
+	TArray<TSharedPtr<FTrackInfo, ESPMode::ThreadSafe>> Tracks;
 	TArray<TArray<int32>> UsableTrackArrayIndicesByType;
 
 	Electra::FTimeFraction MovieDuration;
diff --git a/Engine/Plugins/Media/ElectraPlayer/Source/ElectraProtron/Private/Player/ElectraProtronPlayerLoader.cpp b/Engine/Plugins/Media/ElectraPlayer/Source/ElectraProtron/Private/Player/ElectraProtronPlayerLoader.cpp
--- a/Engine/Plugins/Media/ElectraPlayer/Source/ElectraProtron/Private/Player/ElectraProtronPlayerLoader.cpp
+++ b/Engine/Plugins/Media/ElectraPlayer/Source/ElectraProtron/Private/Player/ElectraProtronPlayerLoader.cpp
@@ -277,7 +277,7 @@ FElectraProtronPlayer::FImpl::FMP4SamplePtr FElectraProtronPlayer::FImpl::FLoade
 	LoadRequestLock.Lock();
 	if (ActiveLoadRequest.TrackSampleBuffer == InFromBuffer)
 	{
-		ActiveLoadRequest.UpdateAtIterator = InAtIterator;
+		ActiveLoadRequest.UpdateAtIterator = InAtIterator->Clone();
 		LoadRequestDirty = (int32) SampleNum;
 	}
 	LoadRequestLock.Unlock();
@@ -467,10 +467,21 @@ FElectraProtronPlayer::FImpl::FLoaderThread::ELoadResult FElectraProtronPlayer::
 
 TSharedPtr<FElectraProtronPlayer::FImpl::FMP4Sample, ESPMode::ThreadSafe> FElectraProtronPlayer::FImpl::FLoaderThread::RetrieveSample(const FMP4TrackSampleBufferPtr& InTrackSampleBuffer, const FTrackIterator& InSampleIt, const FTrackIterator& InOptionalTimecodeIt, const ElectraProtronUtils::FCodecInfo::FTMCDTimecode& InTimecodeInfo)
 {
-	// Note: As long as we are not adding or removing samples in another thread, a lock should not be necessary here.
+	InTrackSampleBuffer->Lock.Lock();
 	FMP4SamplePtr* CurrentSample = InTrackSampleBuffer->SampleMap.Find(InSampleIt->GetSampleNumber());
+	InTrackSampleBuffer->Lock.Unlock();
 	if (CurrentSample)
 	{
+		check((*CurrentSample)->DTS == InSampleIt->GetDTS().GetAsTimespan());
+		check((*CurrentSample)->PTS == InSampleIt->GetPTS().GetAsTimespan());
+		check((*CurrentSample)->EffectiveDTS == InSampleIt->GetEffectiveDTS().GetAsTimespan());
+		check((*CurrentSample)->EffectivePTS == InSampleIt->GetEffectivePTS().GetAsTimespan());
+		check((*CurrentSample)->Duration == InSampleIt->GetDurationAsTimespan());
+		check((*CurrentSample)->SizeInBytes == InSampleIt->GetSampleSize());
+		check((*CurrentSample)->OffsetInFile == InSampleIt->GetSampleFileOffset());
+		check((*CurrentSample)->TrackID == InSampleIt->GetTrackID());
+		check((*CurrentSample)->SampleNumber == InSampleIt->GetSampleNumber());
+		check((*CurrentSample)->bIsSyncOrRap == InSampleIt->IsSyncOrRAPSample());
 		return *CurrentSample;
 	}
 
diff --git a/Engine/Plugins/Media/ElectraUtil/Source/ElectraBase/Private/Utilities/MP4Boxes/MP4Boxes.cpp b/Engine/Plugins/Media/ElectraUtil/Source/ElectraBase/Private/Utilities/MP4Boxes/MP4Boxes.cpp
--- a/Engine/Plugins/Media/ElectraUtil/Source/ElectraBase/Private/Utilities/MP4Boxes/MP4Boxes.cpp
+++ b/Engine/Plugins/Media/ElectraUtil/Source/ElectraBase/Private/Utilities/MP4Boxes/MP4Boxes.cpp
@@ -16,7 +16,7 @@ namespace Electra
 			FMP4BoxInfo bi;
 			while(InOutReader.ParseIntoBoxInfo(bi, InCurrentBoxInfo.Offset + InCurrentBoxInfo.DataOffset + InOutReader.GetCurrentOffset()))
 			{
-				TSharedPtr<FMP4BoxBase> Child = FMP4BoxFactory::Get().Create(AsWeak(), bi);
+				TSharedPtr<FMP4BoxBase, ESPMode::ThreadSafe> Child = FMP4BoxFactory::Get().Create(AsWeak(), bi);
 				if (AddChildBox(Child) && !Child->IsLeafBox() && !Child->IsListOfEntries())
 				{
 					FMP4AtomReaderHelper ar(Child->BoxInfo.Data);
@@ -466,7 +466,7 @@ namespace Electra
 					UE_LOG(LogElectraBase, Error, TEXT("Failed to parse `stsd` box!"));
 					return;
 				}
-				TSharedPtr<FMP4BoxBase> SampleEntry;
+				TSharedPtr<FMP4BoxBase, ESPMode::ThreadSafe> SampleEntry;
 				switch(Parsed->SampleType)
 				{
 					case FMP4BoxSampleEntry::ESampleType::Video:
@@ -2216,7 +2216,7 @@ namespace Electra
 
 		struct FMP4BoxSENC::FInitRefs
 		{
-			TArray<TSharedPtr<Electra::UtilitiesMP4::FMP4BoxBase>> RelatedBoxes;
+			TArray<TSharedPtr<Electra::UtilitiesMP4::FMP4BoxBase, ESPMode::ThreadSafe>> RelatedBoxes;
 		};
 
 		struct FMP4BoxSENC::FParsed
@@ -2258,7 +2258,7 @@ namespace Electra
 				}
 
 				// We need the `Per_Sample_IV_Size` value from the `tenc` box.
-				TSharedPtr<Electra::UtilitiesMP4::FMP4BoxBase>* TencBoxPointer = InitRefs->RelatedBoxes.FindByPredicate([](const TSharedPtr<Electra::UtilitiesMP4::FMP4BoxBase>& InBox)
+				TSharedPtr<Electra::UtilitiesMP4::FMP4BoxBase, ESPMode::ThreadSafe>* TencBoxPointer = InitRefs->RelatedBoxes.FindByPredicate([](const TSharedPtr<Electra::UtilitiesMP4::FMP4BoxBase, ESPMode::ThreadSafe>& InBox)
 				{
 					return InBox->GetType() == Electra::UtilitiesMP4::MakeBoxAtom('t','e','n','c');
 				});
@@ -2267,7 +2267,7 @@ namespace Electra
 					UE_LOG(LogElectraBase, Error, TEXT("Parsing a `senc` box requires the `tenc` box of the related track to be first passed via Prepare()"));
 					return;
 				}
-				TSharedPtr<Electra::UtilitiesMP4::FMP4BoxTENC> TencBox = StaticCastSharedPtr<Electra::UtilitiesMP4::FMP4BoxTENC>(*TencBoxPointer);
+				TSharedPtr<Electra::UtilitiesMP4::FMP4BoxTENC, ESPMode::ThreadSafe> TencBox = StaticCastSharedPtr<Electra::UtilitiesMP4::FMP4BoxTENC>(*TencBoxPointer);
 				const uint8 Per_Sample_IV_Size = TencBox->GetDefaultPerSampleIVSize();
 
 				const bool senc_use_subsamples = !!(Parsed->Flags & 2);
@@ -2321,7 +2321,7 @@ namespace Electra
 			}
 		}
 
-		void FMP4BoxSENC::Prepare(const TArray<TSharedPtr<Electra::UtilitiesMP4::FMP4BoxBase>>& InRelatedBoxes)
+		void FMP4BoxSENC::Prepare(const TArray<TSharedPtr<Electra::UtilitiesMP4::FMP4BoxBase, ESPMode::ThreadSafe>>& InRelatedBoxes)
 		{
 			// If already prepared we're done.
 			if (!InitRefs.IsValid())
@@ -2558,7 +2558,7 @@ namespace Electra
 			FactoryMap.Add(MakeBoxAtom('d','O','p','s'), &FMP4BoxDOPS::Create);
 			FactoryMap.Add(MakeBoxAtom('v','p','c','C'), &FMP4BoxVPCC::Create);
 		}
-		TSharedPtr<FMP4BoxBase> FMP4BoxFactory::Create(const TWeakPtr<FMP4BoxBase>& InParent, const FMP4BoxInfo& InBoxInfo)
+		TSharedPtr<FMP4BoxBase, ESPMode::ThreadSafe> FMP4BoxFactory::Create(const TWeakPtr<FMP4BoxBase, ESPMode::ThreadSafe>& InParent, const FMP4BoxInfo& InBoxInfo)
 		{
 			if (auto Factory = FactoryMap.Find(InBoxInfo.Type))
 			{
diff --git a/Engine/Plugins/Media/ElectraUtil/Source/ElectraBase/Private/Utilities/MP4Boxes/MP4Track.cpp b/Engine/Plugins/Media/ElectraUtil/Source/ElectraBase/Private/Utilities/MP4Boxes/MP4Track.cpp
--- a/Engine/Plugins/Media/ElectraUtil/Source/ElectraBase/Private/Utilities/MP4Boxes/MP4Track.cpp
+++ b/Engine/Plugins/Media/ElectraUtil/Source/ElectraBase/Private/Utilities/MP4Boxes/MP4Track.cpp
@@ -26,19 +26,19 @@ namespace Electra
 			Convs.DurationFromTkhdBox.SetFromND(TkhdBox->GetDuration(), Convs.DurationFromMvhdBox.GetDenominator());
 
 			// Check for correct box hierarchy.
-			TSharedPtr<FMP4BoxMDIA> MdiaBox = TrakBox->FindBoxRecursive<Electra::UtilitiesMP4::FMP4BoxMDIA>(Electra::UtilitiesMP4::MakeBoxAtom('m','d','i','a'), 0);
+			TSharedPtr<FMP4BoxMDIA, ESPMode::ThreadSafe> MdiaBox = TrakBox->FindBoxRecursive<Electra::UtilitiesMP4::FMP4BoxMDIA>(Electra::UtilitiesMP4::MakeBoxAtom('m','d','i','a'), 0);
 			if (!MdiaBox.IsValid())
 			{
 				LastErrorMessage = FString::Printf(TEXT("No `mdia` box in `trak`."));
 				return false;
 			}
-			TSharedPtr<FMP4BoxMINF> MinfBox = MdiaBox->FindBoxRecursive<Electra::UtilitiesMP4::FMP4BoxMINF>(Electra::UtilitiesMP4::MakeBoxAtom('m','i','n','f'), 0);
+			TSharedPtr<FMP4BoxMINF, ESPMode::ThreadSafe> MinfBox = MdiaBox->FindBoxRecursive<Electra::UtilitiesMP4::FMP4BoxMINF>(Electra::UtilitiesMP4::MakeBoxAtom('m','i','n','f'), 0);
 			if (!MinfBox.IsValid())
 			{
 				LastErrorMessage = FString::Printf(TEXT("No `minf` box in `mdia`."));
 				return false;
 			}
-			TSharedPtr<FMP4BoxSTBL> StblBox = MinfBox->FindBoxRecursive<Electra::UtilitiesMP4::FMP4BoxSTBL>(Electra::UtilitiesMP4::MakeBoxAtom('s','t','b','l'), 0);
+			TSharedPtr<FMP4BoxSTBL, ESPMode::ThreadSafe> StblBox = MinfBox->FindBoxRecursive<Electra::UtilitiesMP4::FMP4BoxSTBL>(Electra::UtilitiesMP4::MakeBoxAtom('s','t','b','l'), 0);
 			if (!StblBox.IsValid())
 			{
 				LastErrorMessage = FString::Printf(TEXT("No `stbl` box in `minf`."));
@@ -157,13 +157,13 @@ namespace Electra
 			if (UdtaBox.IsValid())
 			{
 				// Is there a `name` box?
-				TSharedPtr<FMP4BoxBase> NameBox = UdtaBox->FindBoxRecursive<Electra::UtilitiesMP4::FMP4BoxBase>(Electra::UtilitiesMP4::MakeBoxAtom('n','a','m','e'), 0);
+				TSharedPtr<FMP4BoxBase, ESPMode::ThreadSafe> NameBox = UdtaBox->FindBoxRecursive<Electra::UtilitiesMP4::FMP4BoxBase>(Electra::UtilitiesMP4::MakeBoxAtom('n','a','m','e'), 0);
 				if (NameBox.IsValid() && NameBox->GetBoxData().Num())
 				{
 					CommonMetadata.Name = FString::ConstructFromPtrSize(reinterpret_cast<const ANSICHAR*>(NameBox->GetBoxData().GetData()), NameBox->GetBoxData().Num());
 				}
 			}
-			TSharedPtr<FMP4BoxHDLR> HdlrBox = MdiaBox->FindBoxRecursive<Electra::UtilitiesMP4::FMP4BoxHDLR>(Electra::UtilitiesMP4::MakeBoxAtom('h','d','l','r'), 0);
+			TSharedPtr<FMP4BoxHDLR, ESPMode::ThreadSafe> HdlrBox = MdiaBox->FindBoxRecursive<Electra::UtilitiesMP4::FMP4BoxHDLR>(Electra::UtilitiesMP4::MakeBoxAtom('h','d','l','r'), 0);
 			check(HdlrBox.IsValid())
 			if (HdlrBox.IsValid())
 			{
@@ -172,7 +172,7 @@ namespace Electra
 			// Get language from `mdhd` box first.
 			CommonMetadata.Language = MdhdBox->GetLanguage();
 			// Then, if there is an `elng` box that gives better information get it from there.
-			TSharedPtr<FMP4BoxELNG> ElngBox = MdiaBox->FindBoxRecursive<Electra::UtilitiesMP4::FMP4BoxELNG>(Electra::UtilitiesMP4::MakeBoxAtom('e','l','n','g'), 0);
+			TSharedPtr<FMP4BoxELNG, ESPMode::ThreadSafe> ElngBox = MdiaBox->FindBoxRecursive<Electra::UtilitiesMP4::FMP4BoxELNG>(Electra::UtilitiesMP4::MakeBoxAtom('e','l','n','g'), 0);
 			if (ElngBox.IsValid())
 			{
 				CommonMetadata.Language = ElngBox->GetLanguage();
@@ -182,7 +182,7 @@ namespace Electra
 
 			// Given the timeline mapping, locate the sample number that falls onto the start of the timeline
 			// and last one falling onto the end of the timeline.
-			TSharedPtr<FIterator> StartIt = CreateIterator();
+			TSharedPtr<FIterator, ESPMode::ThreadSafe> StartIt = CreateIterator();
 			check(StartIt.IsValid());
 			const int64 pts0 = Convs.CompositionTimeAtZeroPoint;
 			uint32 SyncSampleNum = 0;
@@ -212,7 +212,7 @@ namespace Electra
 			}
 			check(bFound);
 
-			TSharedPtr<FIterator> EndIt = CreateIteratorAtLastFrame();
+			TSharedPtr<FIterator, ESPMode::ThreadSafe> EndIt = CreateIteratorAtLastFrame();
 			check(EndIt.IsValid());
 
 			// Find the highest PTS
@@ -317,7 +317,7 @@ namespace Electra
 		}
 
 		// Create an interator starting at the first sample.
-		TSharedPtr<FMP4Track::FIterator> FMP4Track::CreateIterator()
+		TSharedPtr<FMP4Track::FIterator, ESPMode::ThreadSafe> FMP4Track::CreateIterator()
 		{
 			check(bHasBeenPrepared);
 			if (!bHasBeenPrepared)
@@ -343,8 +343,8 @@ namespace Electra
 			It->stscIt.SetBox(StscBox, Convs.NumTotalSamples);
 			It->stssIt.SetBox(StssBox, Convs.NumTotalSamples);
 			// Do we have a `rap ` group?
-			TSharedPtr<FMP4BoxSGPD>* rapSGPD = SgpdBoxes.FindByPredicate([](const TSharedPtr<FMP4BoxSGPD>& e){ return e->GetGroupingType() == Electra::UtilitiesMP4::MakeBoxAtom('r','a','p',' ');});
-			TSharedPtr<FMP4BoxSBGP>* rapSBGP = SbgpBoxes.FindByPredicate([](const TSharedPtr<FMP4BoxSBGP>& e){ return e->GetGroupingType() == Electra::UtilitiesMP4::MakeBoxAtom('r','a','p',' ');});
+			TSharedPtr<FMP4BoxSGPD, ESPMode::ThreadSafe>* rapSGPD = SgpdBoxes.FindByPredicate([](const TSharedPtr<FMP4BoxSGPD, ESPMode::ThreadSafe>& e){ return e->GetGroupingType() == Electra::UtilitiesMP4::MakeBoxAtom('r','a','p',' ');});
+			TSharedPtr<FMP4BoxSBGP, ESPMode::ThreadSafe>* rapSBGP = SbgpBoxes.FindByPredicate([](const TSharedPtr<FMP4BoxSBGP, ESPMode::ThreadSafe>& e){ return e->GetGroupingType() == Electra::UtilitiesMP4::MakeBoxAtom('r','a','p',' ');});
 			if (rapSGPD && rapSBGP)
 			{
 				It->rapIt.SetBox(*rapSBGP, (*rapSGPD)->GetDefaultGroupDescriptionIndex(), Convs.NumTotalSamples);
@@ -359,7 +359,7 @@ namespace Electra
 		}
 
 		// Create an interator starting at the last sample (used when iterating in reverse, crossing back from the beginning to the end)
-		TSharedPtr<FMP4Track::FIterator> FMP4Track::CreateIteratorAtLastFrame()
+		TSharedPtr<FMP4Track::FIterator, ESPMode::ThreadSafe> FMP4Track::CreateIteratorAtLastFrame()
 		{
 			check(bHasBeenPrepared);
 			if (!bHasBeenPrepared)
@@ -375,9 +375,9 @@ namespace Electra
 			return CreateIterator(Convs.NumTotalSamples - 1);
 		}
 
-		TSharedPtr<FMP4Track::FIterator> FMP4Track::CreateIterator(uint32 InAtSampleNumber)
+		TSharedPtr<FMP4Track::FIterator, ESPMode::ThreadSafe> FMP4Track::CreateIterator(uint32 InAtSampleNumber)
 		{
-			TSharedPtr<FMP4Track::FIterator> It(CreateIterator());
+			TSharedPtr<FMP4Track::FIterator, ESPMode::ThreadSafe> It(CreateIterator());
 			if (It.IsValid() && InAtSampleNumber)
 			{
 				InAtSampleNumber = InAtSampleNumber < Convs.NumTotalSamples ? InAtSampleNumber : Convs.NumTotalSamples-1;
@@ -395,7 +395,7 @@ namespace Electra
 		}
 
 
-		TSharedPtr<FMP4Track::FIterator> FMP4Track::CreateIteratorAtKeyframe(FTimeValue InForTime, FTimeValue InLaterTimeThreshold)
+		TSharedPtr<FMP4Track::FIterator, ESPMode::ThreadSafe> FMP4Track::CreateIteratorAtKeyframe(FTimeValue InForTime, FTimeValue InLaterTimeThreshold)
 		{
 			check(bHasBeenPrepared);
 			if (!bHasBeenPrepared || !SttsBox.IsValid() || SttsBox->GetTotalDuration()==0)
@@ -437,7 +437,7 @@ namespace Electra
 			check(ApproxSampleNumber <= 0xffffffff);
 
 
-			TSharedPtr<FMP4Track::FIterator> ApproxFrameIt(CreateIterator((uint32)ApproxSampleNumber));
+			TSharedPtr<FMP4Track::FIterator, ESPMode::ThreadSafe> ApproxFrameIt(CreateIterator((uint32)ApproxSampleNumber));
 			if (!ApproxFrameIt.IsValid())
 			{
 				LastErrorMessage = FString::Printf(TEXT("Failed to create track iterator for sample #%lld with %u samples in track"), (long long int)ApproxSampleNumber, Convs.NumTotalSamples);
@@ -457,10 +457,10 @@ namespace Electra
 				{ }
 			}
 			// Locate the nearest earlier sync sample, which might be the current one already.
-			TSharedPtr<FMP4Track::FIterator> SyncFrameIt(ApproxFrameIt->Clone());
+			TSharedPtr<FMP4Track::FIterator, ESPMode::ThreadSafe> SyncFrameIt(ApproxFrameIt->Clone());
 			for(; !SyncFrameIt->IsFirst() && (SyncFrameIt->GetPTS().GetNumerator() > LocalTrackTime || !SyncFrameIt->IsSyncOrRAPSample()); SyncFrameIt->Prev())
 			{ }
-			TSharedPtr<FMP4Track::FIterator> NextSyncFrameIt(ApproxFrameIt->Clone());
+			TSharedPtr<FMP4Track::FIterator, ESPMode::ThreadSafe> NextSyncFrameIt(ApproxFrameIt->Clone());
 			bool bLaterOneIsPossible = false;
 			if (MaxLocalTrackTime > LocalTrackTime)
 			{
diff --git a/Engine/Plugins/Media/ElectraUtil/Source/ElectraBase/Public/Utilities/MP4Boxes/MP4BoxBase.h b/Engine/Plugins/Media/ElectraUtil/Source/ElectraBase/Public/Utilities/MP4Boxes/MP4BoxBase.h
--- a/Engine/Plugins/Media/ElectraUtil/Source/ElectraBase/Public/Utilities/MP4Boxes/MP4BoxBase.h
+++ b/Engine/Plugins/Media/ElectraUtil/Source/ElectraBase/Public/Utilities/MP4Boxes/MP4BoxBase.h
@@ -11,7 +11,7 @@ namespace Electra
 		class FMP4BoxBase : public TSharedFromThis<FMP4BoxBase>
 		{
 		public:
-			static TSharedPtr<FMP4BoxBase> Create(const TWeakPtr<FMP4BoxBase>& InParent, const FMP4BoxInfo& InBoxInfo)
+			static TSharedPtr<FMP4BoxBase, ESPMode::ThreadSafe> Create(const TWeakPtr<FMP4BoxBase, ESPMode::ThreadSafe>& InParent, const FMP4BoxInfo& InBoxInfo)
 			{ return MakeShareable<FMP4BoxBase>(new FMP4BoxBase(InParent, InBoxInfo)); }
 			virtual ~FMP4BoxBase() = default;
 
@@ -38,9 +38,9 @@ namespace Electra
 			// Parent box query methods
 			//
 			template<typename T>
-			TSharedPtr<T> FindParentBox(uint32 InType) const
+			TSharedPtr<T, ESPMode::ThreadSafe> FindParentBox(uint32 InType) const
 			{
-				TSharedPtr<T> Parent = StaticCastSharedPtr<T>(ParentBox.Pin());
+				TSharedPtr<T, ESPMode::ThreadSafe> Parent = StaticCastSharedPtr<T>(ParentBox.Pin());
 				while(Parent.IsValid())
 				{
 					if (Parent->GetType() == InType)
@@ -55,12 +55,12 @@ namespace Electra
 			//-------------------------------------------------------------------------
 			// Child box query methods
 			//
-			const TArray<TSharedPtr<FMP4BoxBase>> GetChildren() const
+			const TArray<TSharedPtr<FMP4BoxBase, ESPMode::ThreadSafe>> GetChildren() const
 			{ return Children; }
 
 			// Find the first box of the given type being a child of this box, or any child of a child box.
 			template<typename T>
-			TSharedPtr<T> FindBoxRecursive(uint32 InType, int32 InMaxDepth=8) const
+			TSharedPtr<T, ESPMode::ThreadSafe> FindBoxRecursive(uint32 InType, int32 InMaxDepth=8) const
 			{
 				// First pass, check all children if they are of the type
 				for(int32 i=0,iMax=Children.Num(); i<iMax; ++i)
@@ -75,7 +75,7 @@ namespace Electra
 				{
 					for(int32 i=0, iMax=Children.Num(); i<iMax; ++i)
 					{
-						TSharedPtr<T> Box = Children[i]->FindBoxRecursive<T>(InType, InMaxDepth - 1);
+						TSharedPtr<T, ESPMode::ThreadSafe> Box = Children[i]->FindBoxRecursive<T>(InType, InMaxDepth - 1);
 						if (Box.IsValid())
 						{
 							return Box;
@@ -88,7 +88,7 @@ namespace Electra
 
 			// Returns all instances of a given box type from the direct children of THIS BOX ONLY.
 			template<typename T>
-			void GetAllBoxInstances(TArray<TSharedPtr<T>>& OutAllBoxes, uint32 InType) const
+			void GetAllBoxInstances(TArray<TSharedPtr<T, ESPMode::ThreadSafe>>& OutAllBoxes, uint32 InType) const
 			{
 				for(int32 i=0,iMax=Children.Num(); i<iMax; ++i)
 				{
@@ -102,7 +102,7 @@ namespace Electra
 
 
 			// Called by building the box tree. Not to be called by user code.
-			virtual bool AddChildBox(TSharedPtr<FMP4BoxBase> InChildBox)
+			virtual bool AddChildBox(TSharedPtr<FMP4BoxBase, ESPMode::ThreadSafe> InChildBox)
 			{
 				if (InChildBox.IsValid())
 				{
@@ -117,12 +117,12 @@ namespace Electra
 			void SetRootBoxData(const TSharedPtr<FMP4BoxInfo, ESPMode::ThreadSafe>& InRootBoxData)
 			{ RootBoxData = InRootBoxData; }
 		protected:
-			FMP4BoxBase(const TWeakPtr<FMP4BoxBase>& InParent, const FMP4BoxInfo& InBoxInfo) : ParentBox(InParent), BoxInfo(InBoxInfo)
+			FMP4BoxBase(const TWeakPtr<FMP4BoxBase, ESPMode::ThreadSafe>& InParent, const FMP4BoxInfo& InBoxInfo) : ParentBox(InParent), BoxInfo(InBoxInfo)
 			{ }
 
 			TSharedPtr<FMP4BoxInfo, ESPMode::ThreadSafe> RootBoxData;
-			TWeakPtr<FMP4BoxBase> ParentBox;
-			TArray<TSharedPtr<FMP4BoxBase>> Children;
+			TWeakPtr<FMP4BoxBase, ESPMode::ThreadSafe> ParentBox;
+			TArray<TSharedPtr<FMP4BoxBase, ESPMode::ThreadSafe>> Children;
 			FMP4BoxInfo BoxInfo;
 		};
 
diff --git a/Engine/Plugins/Media/ElectraUtil/Source/ElectraBase/Public/Utilities/MP4Boxes/MP4BoxIterators.h b/Engine/Plugins/Media/ElectraUtil/Source/ElectraBase/Public/Utilities/MP4Boxes/MP4BoxIterators.h
--- a/Engine/Plugins/Media/ElectraUtil/Source/ElectraBase/Public/Utilities/MP4Boxes/MP4BoxIterators.h
+++ b/Engine/Plugins/Media/ElectraUtil/Source/ElectraBase/Public/Utilities/MP4Boxes/MP4BoxIterators.h
@@ -13,7 +13,7 @@ namespace Electra
 
 		struct FSTTSBoxIterator
 		{
-			void SetBox(const TSharedPtr<FMP4BoxSTTS>& InSttsBox)
+			void SetBox(const TSharedPtr<FMP4BoxSTTS, ESPMode::ThreadSafe>& InSttsBox)
 			{
 				SttsBox = InSttsBox;
 				NumTotalSamples = SttsBox->GetNumTotalSamples();
@@ -109,7 +109,7 @@ namespace Electra
 			{ return CurrentSampleNum + 1 >= NumTotalSamples; }
 
 		private:
-			TSharedPtr<FMP4BoxSTTS> SttsBox;
+			TSharedPtr<FMP4BoxSTTS, ESPMode::ThreadSafe> SttsBox;
 			int64 Time = 0;
 			uint32 CurrentSampleNum = 0;
 			uint32 NumTotalSamples = 0;
@@ -122,7 +122,7 @@ namespace Electra
 
 		struct FCTTSBoxIterator
 		{
-			void SetBox(const TSharedPtr<FMP4BoxCTTS>& InCttsBox, uint32 InMaxSamples)
+			void SetBox(const TSharedPtr<FMP4BoxCTTS, ESPMode::ThreadSafe>& InCttsBox, uint32 InMaxSamples)
 			{
 				if (InCttsBox.IsValid())
 				{
@@ -222,7 +222,7 @@ namespace Electra
 			{ return CurrentSampleNum + 1 >= NumTotalSamples; }
 
 		private:
-			TSharedPtr<FMP4BoxCTTS> CttsBox;
+			TSharedPtr<FMP4BoxCTTS, ESPMode::ThreadSafe> CttsBox;
 			int64 CurrentOffset = 0;
 			uint32 CurrentSampleNum = 0;
 			uint32 NumTotalSamples = 0;
@@ -234,7 +234,7 @@ namespace Electra
 
 		struct FSTSCBoxIterator
 		{
-			void SetBox(const TSharedPtr<FMP4BoxSTSC>& InStscBox, uint32 InMaxSamples)
+			void SetBox(const TSharedPtr<FMP4BoxSTSC, ESPMode::ThreadSafe>& InStscBox, uint32 InMaxSamples)
 			{
 				StscBox = InStscBox;
 				NumTotalSamples = InMaxSamples;
@@ -375,7 +375,7 @@ namespace Electra
 			{ return CurrentSampleNum + 1 >= NumTotalSamples; }
 
 		private:
-			TSharedPtr<FMP4BoxSTSC> StscBox;
+			TSharedPtr<FMP4BoxSTSC, ESPMode::ThreadSafe> StscBox;
 			uint32 Current_first_chunk = 0;
 			uint32 Current_samples_per_chunk = 0;
 			uint32 Current_sample_description_index = 0;
@@ -390,7 +390,7 @@ namespace Electra
 
 		struct FSTSZBoxIterator
 		{
-			void SetBox(const TSharedPtr<FMP4BoxSTSZ>& InStszBox)
+			void SetBox(const TSharedPtr<FMP4BoxSTSZ, ESPMode::ThreadSafe>& InStszBox)
 			{
 				check(InStszBox.IsValid());
 				StszBox = InStszBox;
@@ -453,7 +453,7 @@ namespace Electra
 			{ return CurrentSampleNum + 1 >= NumTotalSamples; }
 
 		private:
-			TSharedPtr<FMP4BoxSTSZ> StszBox;
+			TSharedPtr<FMP4BoxSTSZ, ESPMode::ThreadSafe> StszBox;
 			uint32 NumTotalSamples = 0;
 			uint32 CurrentSampleNum = 0;
 		};
@@ -463,7 +463,7 @@ namespace Electra
 		// This is not so much an iterator than a holder of the box.
 		struct FSTCOBoxIterator
 		{
-			void SetBox(const TSharedPtr<FMP4BoxSTCO>& InStcoBox)
+			void SetBox(const TSharedPtr<FMP4BoxSTCO, ESPMode::ThreadSafe>& InStcoBox)
 			{
 				check(InStcoBox.IsValid());
 				StcoBox = InStcoBox;
@@ -478,14 +478,14 @@ namespace Electra
 				return StcoBox.IsValid() ? StcoBox->GetChunkOffset(InChunkIndex) : 0;
 			}
 		private:
-			TSharedPtr<FMP4BoxSTCO> StcoBox;
+			TSharedPtr<FMP4BoxSTCO, ESPMode::ThreadSafe> StcoBox;
 		};
 
 		/****************************************************************************************************************************************************/
 
 		struct FSTSSBoxIterator
 		{
-			void SetBox(const TSharedPtr<FMP4BoxSTSS>& InStssBox, uint32 InMaxSamples)
+			void SetBox(const TSharedPtr<FMP4BoxSTSS, ESPMode::ThreadSafe>& InStssBox, uint32 InMaxSamples)
 			{
 				StssBox = InStssBox;
 				NumTotalSamples = InMaxSamples;
@@ -572,7 +572,7 @@ namespace Electra
 			{ return CurrentSampleNum + 1 >= NumTotalSamples; }
 
 		private:
-			TSharedPtr<FMP4BoxSTSS> StssBox;
+			TSharedPtr<FMP4BoxSTSS, ESPMode::ThreadSafe> StssBox;
 			uint32 NumTotalSamples = 0;
 			uint32 CurrentSampleNum = 0;
 			uint32 CurrentEntryIndex = 0;
@@ -583,7 +583,7 @@ namespace Electra
 
 		struct FSBGPBoxIterator
 		{
-			void SetBox(const TSharedPtr<FMP4BoxSBGP>& InSbgpBox, uint32 InDefaultSampleGroupDescriptionIndex, uint32 InMaxSamples)
+			void SetBox(const TSharedPtr<FMP4BoxSBGP, ESPMode::ThreadSafe>& InSbgpBox, uint32 InDefaultSampleGroupDescriptionIndex, uint32 InMaxSamples)
 			{
 				NumTotalSamples = InMaxSamples;
 				Current_sample_count = InMaxSamples;
@@ -734,7 +734,7 @@ namespace Electra
 			{ return CurrentSampleNum + 1 >= NumTotalSamples; }
 
 		private:
-			TSharedPtr<FMP4BoxSBGP> SbgpBox;
+			TSharedPtr<FMP4BoxSBGP, ESPMode::ThreadSafe> SbgpBox;
 			uint32 Current_sample_count = 0;
 			uint32 Current_group_description_index = 0;
 			uint32 DefaultSampleGroupDescriptionIndex = 0;
diff --git a/Engine/Plugins/Media/ElectraUtil/Source/ElectraBase/Public/Utilities/MP4Boxes/MP4Boxes.h b/Engine/Plugins/Media/ElectraUtil/Source/ElectraBase/Public/Utilities/MP4Boxes/MP4Boxes.h
--- a/Engine/Plugins/Media/ElectraUtil/Source/ElectraBase/Public/Utilities/MP4Boxes/MP4Boxes.h
+++ b/Engine/Plugins/Media/ElectraUtil/Source/ElectraBase/Public/Utilities/MP4Boxes/MP4Boxes.h
@@ -18,7 +18,7 @@ namespace Electra
 		class FMP4BoxFREE : public FMP4BoxBase
 		{
 		public:
-			static TSharedPtr<FMP4BoxBase> Create(const TWeakPtr<FMP4BoxBase>& InParent, const FMP4BoxInfo& InBoxInfo)
+			static TSharedPtr<FMP4BoxBase, ESPMode::ThreadSafe> Create(const TWeakPtr<FMP4BoxBase, ESPMode::ThreadSafe>& InParent, const FMP4BoxInfo& InBoxInfo)
 			{ return MakeShareable<FMP4BoxFREE>(new FMP4BoxFREE(InParent, InBoxInfo)); }
 			virtual ~FMP4BoxFREE() = default;
 		protected:
@@ -35,11 +35,11 @@ namespace Electra
 		class FMP4BoxFTYP : public FMP4BoxBase
 		{
 		public:
-			static TSharedPtr<FMP4BoxBase> Create(const TWeakPtr<FMP4BoxBase>& InParent, const FMP4BoxInfo& InBoxInfo)
+			static TSharedPtr<FMP4BoxBase, ESPMode::ThreadSafe> Create(const TWeakPtr<FMP4BoxBase, ESPMode::ThreadSafe>& InParent, const FMP4BoxInfo& InBoxInfo)
 			{ return MakeShareable<FMP4BoxFTYP>(new FMP4BoxFTYP(InParent, InBoxInfo)); }
 			virtual ~FMP4BoxFTYP() = default;
 		protected:
-			FMP4BoxFTYP(const TWeakPtr<FMP4BoxBase>& InParent, const FMP4BoxInfo& InBoxInfo)
+			FMP4BoxFTYP(const TWeakPtr<FMP4BoxBase, ESPMode::ThreadSafe>& InParent, const FMP4BoxInfo& InBoxInfo)
 				: FMP4BoxBase(InParent, InBoxInfo)
 			{ }
 		};
@@ -52,12 +52,12 @@ namespace Electra
 		class FMP4BoxMOOV : public FMP4BoxBase
 		{
 		public:
-			static TSharedPtr<FMP4BoxBase> Create(const TWeakPtr<FMP4BoxBase>& InParent, const FMP4BoxInfo& InBoxInfo)
+			static TSharedPtr<FMP4BoxBase, ESPMode::ThreadSafe> Create(const TWeakPtr<FMP4BoxBase, ESPMode::ThreadSafe>& InParent, const FMP4BoxInfo& InBoxInfo)
 			{ return MakeShareable<FMP4BoxMOOV>(new FMP4BoxMOOV(InParent, InBoxInfo)); }
 			virtual ~FMP4BoxMOOV() = default;
 			bool IsLeafBox() const override { return false; }
 		protected:
-			FMP4BoxMOOV(const TWeakPtr<FMP4BoxBase>& InParent, const FMP4BoxInfo& InBoxInfo)
+			FMP4BoxMOOV(const TWeakPtr<FMP4BoxBase, ESPMode::ThreadSafe>& InParent, const FMP4BoxInfo& InBoxInfo)
 				: FMP4BoxBase(InParent, InBoxInfo)
 			{ }
 		};
@@ -70,14 +70,14 @@ namespace Electra
 		class FMP4BoxMVHD : public FMP4BoxBase
 		{
 		public:
-			static TSharedPtr<FMP4BoxBase> Create(const TWeakPtr<FMP4BoxBase>& InParent, const FMP4BoxInfo& InBoxInfo)
+			static TSharedPtr<FMP4BoxBase, ESPMode::ThreadSafe> Create(const TWeakPtr<FMP4BoxBase, ESPMode::ThreadSafe>& InParent, const FMP4BoxInfo& InBoxInfo)
 			{ return MakeShareable<FMP4BoxMVHD>(new FMP4BoxMVHD(InParent, InBoxInfo)); }
 			virtual ~FMP4BoxMVHD() = default;
 
 			ELECTRABASE_API FTimeFraction GetDuration();
 			ELECTRABASE_API uint32 GetTimescale();
 		protected:
-			FMP4BoxMVHD(const TWeakPtr<FMP4BoxBase>& InParent, const FMP4BoxInfo& InBoxInfo)
+			FMP4BoxMVHD(const TWeakPtr<FMP4BoxBase, ESPMode::ThreadSafe>& InParent, const FMP4BoxInfo& InBoxInfo)
 				: FMP4BoxBase(InParent, InBoxInfo)
 			{ }
 			void ParseIfRequired();
@@ -93,12 +93,12 @@ namespace Electra
 		class FMP4BoxTRAK : public FMP4BoxBase
 		{
 		public:
-			static TSharedPtr<FMP4BoxBase> Create(const TWeakPtr<FMP4BoxBase>& InParent, const FMP4BoxInfo& InBoxInfo)
+			static TSharedPtr<FMP4BoxBase, ESPMode::ThreadSafe> Create(const TWeakPtr<FMP4BoxBase, ESPMode::ThreadSafe>& InParent, const FMP4BoxInfo& InBoxInfo)
 			{ return MakeShareable<FMP4BoxTRAK>(new FMP4BoxTRAK(InParent, InBoxInfo)); }
 			virtual ~FMP4BoxTRAK() = default;
 			bool IsLeafBox() const override { return false; }
 		protected:
-			FMP4BoxTRAK(const TWeakPtr<FMP4BoxBase>& InParent, const FMP4BoxInfo& InBoxInfo)
+			FMP4BoxTRAK(const TWeakPtr<FMP4BoxBase, ESPMode::ThreadSafe>& InParent, const FMP4BoxInfo& InBoxInfo)
 				: FMP4BoxBase(InParent, InBoxInfo)
 			{ }
 		};
@@ -111,7 +111,7 @@ namespace Electra
 		class FMP4BoxTKHD : public FMP4BoxBase
 		{
 		public:
-			static TSharedPtr<FMP4BoxBase> Create(const TWeakPtr<FMP4BoxBase>& InParent, const FMP4BoxInfo& InBoxInfo)
+			static TSharedPtr<FMP4BoxBase, ESPMode::ThreadSafe> Create(const TWeakPtr<FMP4BoxBase, ESPMode::ThreadSafe>& InParent, const FMP4BoxInfo& InBoxInfo)
 			{ return MakeShareable<FMP4BoxTKHD>(new FMP4BoxTKHD(InParent, InBoxInfo)); }
 			virtual ~FMP4BoxTKHD() = default;
 
@@ -140,7 +140,7 @@ namespace Electra
 			ELECTRABASE_API uint16 GetWidth();
 			ELECTRABASE_API uint16 GetHeight();
 		protected:
-			FMP4BoxTKHD(const TWeakPtr<FMP4BoxBase>& InParent, const FMP4BoxInfo& InBoxInfo)
+			FMP4BoxTKHD(const TWeakPtr<FMP4BoxBase, ESPMode::ThreadSafe>& InParent, const FMP4BoxInfo& InBoxInfo)
 				: FMP4BoxBase(InParent, InBoxInfo)
 			{ }
 			void ParseIfRequired();
@@ -156,7 +156,7 @@ namespace Electra
 		class FMP4BoxTREF : public FMP4BoxBase
 		{
 		public:
-			static TSharedPtr<FMP4BoxBase> Create(const TWeakPtr<FMP4BoxBase>& InParent, const FMP4BoxInfo& InBoxInfo)
+			static TSharedPtr<FMP4BoxBase, ESPMode::ThreadSafe> Create(const TWeakPtr<FMP4BoxBase, ESPMode::ThreadSafe>& InParent, const FMP4BoxInfo& InBoxInfo)
 			{ return MakeShareable<FMP4BoxTREF>(new FMP4BoxTREF(InParent, InBoxInfo)); }
 			virtual ~FMP4BoxTREF() = default;
 			bool IsLeafBox() const override { return false; }
@@ -171,7 +171,7 @@ namespace Electra
 			ELECTRABASE_API const TArray<FEntry>& GetEntries();
 			ELECTRABASE_API TArray<FEntry> GetEntriesOfType(uint32 InType);
 		protected:
-			FMP4BoxTREF(const TWeakPtr<FMP4BoxBase>& InParent, const FMP4BoxInfo& InBoxInfo)
+			FMP4BoxTREF(const TWeakPtr<FMP4BoxBase, ESPMode::ThreadSafe>& InParent, const FMP4BoxInfo& InBoxInfo)
 				: FMP4BoxBase(InParent, InBoxInfo)
 			{ }
 			void ParseIfRequired();
@@ -187,12 +187,12 @@ namespace Electra
 		class FMP4BoxEDTS : public FMP4BoxBase
 		{
 		public:
-			static TSharedPtr<FMP4BoxBase> Create(const TWeakPtr<FMP4BoxBase>& InParent, const FMP4BoxInfo& InBoxInfo)
+			static TSharedPtr<FMP4BoxBase, ESPMode::ThreadSafe> Create(const TWeakPtr<FMP4BoxBase, ESPMode::ThreadSafe>& InParent, const FMP4BoxInfo& InBoxInfo)
 			{ return MakeShareable<FMP4BoxEDTS>(new FMP4BoxEDTS(InParent, InBoxInfo)); }
 			virtual ~FMP4BoxEDTS() = default;
 			bool IsLeafBox() const override { return false; }
 		protected:
-			FMP4BoxEDTS(const TWeakPtr<FMP4BoxBase>& InParent, const FMP4BoxInfo& InBoxInfo)
+			FMP4BoxEDTS(const TWeakPtr<FMP4BoxBase, ESPMode::ThreadSafe>& InParent, const FMP4BoxInfo& InBoxInfo)
 				: FMP4BoxBase(InParent, InBoxInfo)
 			{ }
 		};
@@ -205,7 +205,7 @@ namespace Electra
 		class FMP4BoxELST : public FMP4BoxBase
 		{
 		public:
-			static TSharedPtr<FMP4BoxBase> Create(const TWeakPtr<FMP4BoxBase>& InParent, const FMP4BoxInfo& InBoxInfo)
+			static TSharedPtr<FMP4BoxBase, ESPMode::ThreadSafe> Create(const TWeakPtr<FMP4BoxBase, ESPMode::ThreadSafe>& InParent, const FMP4BoxInfo& InBoxInfo)
 			{ return MakeShareable<FMP4BoxELST>(new FMP4BoxELST(InParent, InBoxInfo)); }
 			virtual ~FMP4BoxELST() = default;
 
@@ -223,7 +223,7 @@ namespace Electra
 			ELECTRABASE_API bool RepeatEdits();
 
 		protected:
-			FMP4BoxELST(const TWeakPtr<FMP4BoxBase>& InParent, const FMP4BoxInfo& InBoxInfo)
+			FMP4BoxELST(const TWeakPtr<FMP4BoxBase, ESPMode::ThreadSafe>& InParent, const FMP4BoxInfo& InBoxInfo)
 				: FMP4BoxBase(InParent, InBoxInfo)
 			{ }
 			void ParseIfRequired();
@@ -239,12 +239,12 @@ namespace Electra
 		class FMP4BoxMDIA : public FMP4BoxBase
 		{
 		public:
-			static TSharedPtr<FMP4BoxBase> Create(const TWeakPtr<FMP4BoxBase>& InParent, const FMP4BoxInfo& InBoxInfo)
+			static TSharedPtr<FMP4BoxBase, ESPMode::ThreadSafe> Create(const TWeakPtr<FMP4BoxBase, ESPMode::ThreadSafe>& InParent, const FMP4BoxInfo& InBoxInfo)
 			{ return MakeShareable<FMP4BoxMDIA>(new FMP4BoxMDIA(InParent, InBoxInfo)); }
 			virtual ~FMP4BoxMDIA() = default;
 			bool IsLeafBox() const override { return false; }
 		protected:
-			FMP4BoxMDIA(const TWeakPtr<FMP4BoxBase>& InParent, const FMP4BoxInfo& InBoxInfo)
+			FMP4BoxMDIA(const TWeakPtr<FMP4BoxBase, ESPMode::ThreadSafe>& InParent, const FMP4BoxInfo& InBoxInfo)
 				: FMP4BoxBase(InParent, InBoxInfo)
 			{ }
 		};
@@ -257,7 +257,7 @@ namespace Electra
 		class FMP4BoxMDHD : public FMP4BoxBase
 		{
 		public:
-			static TSharedPtr<FMP4BoxBase> Create(const TWeakPtr<FMP4BoxBase>& InParent, const FMP4BoxInfo& InBoxInfo)
+			static TSharedPtr<FMP4BoxBase, ESPMode::ThreadSafe> Create(const TWeakPtr<FMP4BoxBase, ESPMode::ThreadSafe>& InParent, const FMP4BoxInfo& InBoxInfo)
 			{ return MakeShareable<FMP4BoxMDHD>(new FMP4BoxMDHD(InParent, InBoxInfo)); }
 			virtual ~FMP4BoxMDHD() = default;
 
@@ -266,7 +266,7 @@ namespace Electra
 			ELECTRABASE_API const FString& GetLanguage();
 
 		protected:
-			FMP4BoxMDHD(const TWeakPtr<FMP4BoxBase>& InParent, const FMP4BoxInfo& InBoxInfo)
+			FMP4BoxMDHD(const TWeakPtr<FMP4BoxBase, ESPMode::ThreadSafe>& InParent, const FMP4BoxInfo& InBoxInfo)
 				: FMP4BoxBase(InParent, InBoxInfo)
 			{ }
 			void ParseIfRequired();
@@ -282,14 +282,14 @@ namespace Electra
 		class FMP4BoxHDLR : public FMP4BoxBase
 		{
 		public:
-			static TSharedPtr<FMP4BoxBase> Create(const TWeakPtr<FMP4BoxBase>& InParent, const FMP4BoxInfo& InBoxInfo)
+			static TSharedPtr<FMP4BoxBase, ESPMode::ThreadSafe> Create(const TWeakPtr<FMP4BoxBase, ESPMode::ThreadSafe>& InParent, const FMP4BoxInfo& InBoxInfo)
 			{ return MakeShareable<FMP4BoxHDLR>(new FMP4BoxHDLR(InParent, InBoxInfo)); }
 			virtual ~FMP4BoxHDLR() = default;
 
 			ELECTRABASE_API uint32 GetHandlerType();
 			ELECTRABASE_API FString GetHandlerName();
 		protected:
-			FMP4BoxHDLR(const TWeakPtr<FMP4BoxBase>& InParent, const FMP4BoxInfo& InBoxInfo)
+			FMP4BoxHDLR(const TWeakPtr<FMP4BoxBase, ESPMode::ThreadSafe>& InParent, const FMP4BoxInfo& InBoxInfo)
 				: FMP4BoxBase(InParent, InBoxInfo)
 			{ }
 			void ParseIfRequired();
@@ -305,12 +305,12 @@ namespace Electra
 		class FMP4BoxMINF : public FMP4BoxBase
 		{
 		public:
-			static TSharedPtr<FMP4BoxBase> Create(const TWeakPtr<FMP4BoxBase>& InParent, const FMP4BoxInfo& InBoxInfo)
+			static TSharedPtr<FMP4BoxBase, ESPMode::ThreadSafe> Create(const TWeakPtr<FMP4BoxBase, ESPMode::ThreadSafe>& InParent, const FMP4BoxInfo& InBoxInfo)
 			{ return MakeShareable<FMP4BoxMINF>(new FMP4BoxMINF(InParent, InBoxInfo)); }
 			virtual ~FMP4BoxMINF() = default;
 			bool IsLeafBox() const override { return false; }
 		protected:
-			FMP4BoxMINF(const TWeakPtr<FMP4BoxBase>& InParent, const FMP4BoxInfo& InBoxInfo)
+			FMP4BoxMINF(const TWeakPtr<FMP4BoxBase, ESPMode::ThreadSafe>& InParent, const FMP4BoxInfo& InBoxInfo)
 				: FMP4BoxBase(InParent, InBoxInfo)
 			{ }
 		};
@@ -323,13 +323,13 @@ namespace Electra
 		class FMP4BoxELNG : public FMP4BoxBase
 		{
 		public:
-			static TSharedPtr<FMP4BoxBase> Create(const TWeakPtr<FMP4BoxBase>& InParent, const FMP4BoxInfo& InBoxInfo)
+			static TSharedPtr<FMP4BoxBase, ESPMode::ThreadSafe> Create(const TWeakPtr<FMP4BoxBase, ESPMode::ThreadSafe>& InParent, const FMP4BoxInfo& InBoxInfo)
 			{ return MakeShareable<FMP4BoxELNG>(new FMP4BoxELNG(InParent, InBoxInfo)); }
 			virtual ~FMP4BoxELNG() = default;
 
 			ELECTRABASE_API FString GetLanguage();
 		protected:
-			FMP4BoxELNG(const TWeakPtr<FMP4BoxBase>& InParent, const FMP4BoxInfo& InBoxInfo)
+			FMP4BoxELNG(const TWeakPtr<FMP4BoxBase, ESPMode::ThreadSafe>& InParent, const FMP4BoxInfo& InBoxInfo)
 				: FMP4BoxBase(InParent, InBoxInfo)
 			{ }
 			void ParseIfRequired();
@@ -345,12 +345,12 @@ namespace Electra
 		class FMP4BoxDINF : public FMP4BoxBase
 		{
 		public:
-			static TSharedPtr<FMP4BoxBase> Create(const TWeakPtr<FMP4BoxBase>& InParent, const FMP4BoxInfo& InBoxInfo)
+			static TSharedPtr<FMP4BoxBase, ESPMode::ThreadSafe> Create(const TWeakPtr<FMP4BoxBase, ESPMode::ThreadSafe>& InParent, const FMP4BoxInfo& InBoxInfo)
 			{ return MakeShareable<FMP4BoxDINF>(new FMP4BoxDINF(InParent, InBoxInfo)); }
 			virtual ~FMP4BoxDINF() = default;
 			bool IsLeafBox() const override { return false; }
 		protected:
-			FMP4BoxDINF(const TWeakPtr<FMP4BoxBase>& InParent, const FMP4BoxInfo& InBoxInfo)
+			FMP4BoxDINF(const TWeakPtr<FMP4BoxBase, ESPMode::ThreadSafe>& InParent, const FMP4BoxInfo& InBoxInfo)
 				: FMP4BoxBase(InParent, InBoxInfo)
 			{ }
 		};
@@ -363,13 +363,13 @@ namespace Electra
 		class FMP4BoxDREF : public FMP4BoxBase
 		{
 		public:
-			static TSharedPtr<FMP4BoxBase> Create(const TWeakPtr<FMP4BoxBase>& InParent, const FMP4BoxInfo& InBoxInfo)
+			static TSharedPtr<FMP4BoxBase, ESPMode::ThreadSafe> Create(const TWeakPtr<FMP4BoxBase, ESPMode::ThreadSafe>& InParent, const FMP4BoxInfo& InBoxInfo)
 			{ return MakeShareable<FMP4BoxDREF>(new FMP4BoxDREF(InParent, InBoxInfo)); }
 			virtual ~FMP4BoxDREF() = default;
 			bool IsLeafBox() const override { return false; }
 			bool IsListOfEntries() const override { return true; }
 		protected:
-			FMP4BoxDREF(const TWeakPtr<FMP4BoxBase>& InParent, const FMP4BoxInfo& InBoxInfo)
+			FMP4BoxDREF(const TWeakPtr<FMP4BoxBase, ESPMode::ThreadSafe>& InParent, const FMP4BoxInfo& InBoxInfo)
 				: FMP4BoxBase(InParent, InBoxInfo)
 			{ }
 		};
@@ -382,12 +382,12 @@ namespace Electra
 		class FMP4BoxSTBL : public FMP4BoxBase
 		{
 		public:
-			static TSharedPtr<FMP4BoxBase> Create(const TWeakPtr<FMP4BoxBase>& InParent, const FMP4BoxInfo& InBoxInfo)
+			static TSharedPtr<FMP4BoxBase, ESPMode::ThreadSafe> Create(const TWeakPtr<FMP4BoxBase, ESPMode::ThreadSafe>& InParent, const FMP4BoxInfo& InBoxInfo)
 			{ return MakeShareable<FMP4BoxSTBL>(new FMP4BoxSTBL(InParent, InBoxInfo)); }
 			virtual ~FMP4BoxSTBL() = default;
 			bool IsLeafBox() const override { return false; }
 		protected:
-			FMP4BoxSTBL(const TWeakPtr<FMP4BoxBase>& InParent, const FMP4BoxInfo& InBoxInfo)
+			FMP4BoxSTBL(const TWeakPtr<FMP4BoxBase, ESPMode::ThreadSafe>& InParent, const FMP4BoxInfo& InBoxInfo)
 				: FMP4BoxBase(InParent, InBoxInfo)
 			{ }
 		};
@@ -418,7 +418,7 @@ namespace Electra
 			virtual uint16 GetDataReferenceIndex() = 0;
 
 		protected:
-			FMP4BoxSampleEntry(const TWeakPtr<FMP4BoxBase>& InParent, const FMP4BoxInfo& InBoxInfo)
+			FMP4BoxSampleEntry(const TWeakPtr<FMP4BoxBase, ESPMode::ThreadSafe>& InParent, const FMP4BoxInfo& InBoxInfo)
 				: FMP4BoxBase(InParent, InBoxInfo)
 			{ }
 		};
@@ -431,14 +431,14 @@ namespace Electra
 		class FMP4BoxBTRT : public FMP4BoxBase
 		{
 		public:
-			static TSharedPtr<FMP4BoxBase> Create(const TWeakPtr<FMP4BoxBase>& InParent, const FMP4BoxInfo& InBoxInfo)
+			static TSharedPtr<FMP4BoxBase, ESPMode::ThreadSafe> Create(const TWeakPtr<FMP4BoxBase, ESPMode::ThreadSafe>& InParent, const FMP4BoxInfo& InBoxInfo)
 			{ return MakeShareable<FMP4BoxBTRT>(new FMP4BoxBTRT(InParent, InBoxInfo)); }
 			virtual ~FMP4BoxBTRT() = default;
 			ELECTRABASE_API uint32 GetBufferSizeDB();
 			ELECTRABASE_API uint32 GetMaxBitrate();
 			ELECTRABASE_API uint32 GetAverageBitrate();
 		protected:
-			FMP4BoxBTRT(const TWeakPtr<FMP4BoxBase>& InParent, const FMP4BoxInfo& InBoxInfo)
+			FMP4BoxBTRT(const TWeakPtr<FMP4BoxBase, ESPMode::ThreadSafe>& InParent, const FMP4BoxInfo& InBoxInfo)
 				: FMP4BoxBase(InParent, InBoxInfo)
 			{ }
 			void ParseIfRequired();
@@ -454,7 +454,7 @@ namespace Electra
 		class FMP4BoxVisualSampleEntry : public FMP4BoxSampleEntry
 		{
 		public:
-			static TSharedPtr<FMP4BoxBase> Create(const TWeakPtr<FMP4BoxBase>& InParent, const FMP4BoxInfo& InBoxInfo)
+			static TSharedPtr<FMP4BoxBase, ESPMode::ThreadSafe> Create(const TWeakPtr<FMP4BoxBase, ESPMode::ThreadSafe>& InParent, const FMP4BoxInfo& InBoxInfo)
 			{ return MakeShareable<FMP4BoxVisualSampleEntry>(new FMP4BoxVisualSampleEntry(InParent, InBoxInfo)); }
 			virtual ~FMP4BoxVisualSampleEntry() = default;
 			ESampleType GetSampleType() override
@@ -465,7 +465,7 @@ namespace Electra
 			ELECTRABASE_API uint16 GetFrameCount();
 			ELECTRABASE_API uint16 GetDepth();
 		protected:
-			FMP4BoxVisualSampleEntry(const TWeakPtr<FMP4BoxBase>& InParent, const FMP4BoxInfo& InBoxInfo)
+			FMP4BoxVisualSampleEntry(const TWeakPtr<FMP4BoxBase, ESPMode::ThreadSafe>& InParent, const FMP4BoxInfo& InBoxInfo)
 				: FMP4BoxSampleEntry(InParent, InBoxInfo)
 			{ }
 			void ParseIfRequired();
@@ -481,7 +481,7 @@ namespace Electra
 		class FMP4BoxAudioSampleEntry : public FMP4BoxSampleEntry
 		{
 		public:
-			static TSharedPtr<FMP4BoxBase> Create(const TWeakPtr<FMP4BoxBase>& InParent, const FMP4BoxInfo& InBoxInfo)
+			static TSharedPtr<FMP4BoxBase, ESPMode::ThreadSafe> Create(const TWeakPtr<FMP4BoxBase, ESPMode::ThreadSafe>& InParent, const FMP4BoxInfo& InBoxInfo)
 			{ return MakeShareable<FMP4BoxAudioSampleEntry>(new FMP4BoxAudioSampleEntry(InParent, InBoxInfo)); }
 			virtual ~FMP4BoxAudioSampleEntry() = default;
 			ESampleType GetSampleType() override
@@ -495,7 +495,7 @@ namespace Electra
 			ELECTRABASE_API uint32 GetConstBytesPerAudioPacket();
 			ELECTRABASE_API uint32 GetConstLPCMFramesPerAudioPacket();
 		protected:
-			FMP4BoxAudioSampleEntry(const TWeakPtr<FMP4BoxBase>& InParent, const FMP4BoxInfo& InBoxInfo)
+			FMP4BoxAudioSampleEntry(const TWeakPtr<FMP4BoxBase, ESPMode::ThreadSafe>& InParent, const FMP4BoxInfo& InBoxInfo)
 				: FMP4BoxSampleEntry(InParent, InBoxInfo)
 			{ }
 			void ParseIfRequired();
@@ -511,7 +511,7 @@ namespace Electra
 		class FMP4BoxQTFFTimecodeSampleEntry : public FMP4BoxSampleEntry
 		{
 		public:
-			static TSharedPtr<FMP4BoxBase> Create(const TWeakPtr<FMP4BoxBase>& InParent, const FMP4BoxInfo& InBoxInfo)
+			static TSharedPtr<FMP4BoxBase, ESPMode::ThreadSafe> Create(const TWeakPtr<FMP4BoxBase, ESPMode::ThreadSafe>& InParent, const FMP4BoxInfo& InBoxInfo)
 			{ return MakeShareable<FMP4BoxQTFFTimecodeSampleEntry>(new FMP4BoxQTFFTimecodeSampleEntry(InParent, InBoxInfo)); }
 			virtual ~FMP4BoxQTFFTimecodeSampleEntry() = default;
 			ESampleType GetSampleType() override
@@ -531,7 +531,7 @@ namespace Electra
 			ELECTRABASE_API uint32 GetFrameDuration();
 			ELECTRABASE_API uint32 GetNumberOfFrames();
 		protected:
-			FMP4BoxQTFFTimecodeSampleEntry(const TWeakPtr<FMP4BoxBase>& InParent, const FMP4BoxInfo& InBoxInfo)
+			FMP4BoxQTFFTimecodeSampleEntry(const TWeakPtr<FMP4BoxBase, ESPMode::ThreadSafe>& InParent, const FMP4BoxInfo& InBoxInfo)
 				: FMP4BoxSampleEntry(InParent, InBoxInfo)
 			{ }
 			void ParseIfRequired();
@@ -546,14 +546,14 @@ namespace Electra
 		class FMP4BoxTX3GSampleEntry : public FMP4BoxSampleEntry
 		{
 		public:
-			static TSharedPtr<FMP4BoxBase> Create(const TWeakPtr<FMP4BoxBase>& InParent, const FMP4BoxInfo& InBoxInfo)
+			static TSharedPtr<FMP4BoxBase, ESPMode::ThreadSafe> Create(const TWeakPtr<FMP4BoxBase, ESPMode::ThreadSafe>& InParent, const FMP4BoxInfo& InBoxInfo)
 			{ return MakeShareable<FMP4BoxTX3GSampleEntry>(new FMP4BoxTX3GSampleEntry(InParent, InBoxInfo)); }
 			virtual ~FMP4BoxTX3GSampleEntry() = default;
 			ESampleType GetSampleType() override
 			{ return ESampleType::Subtitles; }
 			uint16 GetDataReferenceIndex() override;
 		protected:
-			FMP4BoxTX3GSampleEntry(const TWeakPtr<FMP4BoxBase>& InParent, const FMP4BoxInfo& InBoxInfo)
+			FMP4BoxTX3GSampleEntry(const TWeakPtr<FMP4BoxBase, ESPMode::ThreadSafe>& InParent, const FMP4BoxInfo& InBoxInfo)
 				: FMP4BoxSampleEntry(InParent, InBoxInfo)
 			{ }
 			void ParseIfRequired();
@@ -569,11 +569,11 @@ namespace Electra
 		class FMP4BoxPASP : public FMP4BoxBase
 		{
 		public:
-			static TSharedPtr<FMP4BoxBase> Create(const TWeakPtr<FMP4BoxBase>& InParent, const FMP4BoxInfo& InBoxInfo)
+			static TSharedPtr<FMP4BoxBase, ESPMode::ThreadSafe> Create(const TWeakPtr<FMP4BoxBase, ESPMode::ThreadSafe>& InParent, const FMP4BoxInfo& InBoxInfo)
 			{ return MakeShareable<FMP4BoxPASP>(new FMP4BoxPASP(InParent, InBoxInfo)); }
 			virtual ~FMP4BoxPASP() = default;
 		protected:
-			FMP4BoxPASP(const TWeakPtr<FMP4BoxBase>& InParent, const FMP4BoxInfo& InBoxInfo)
+			FMP4BoxPASP(const TWeakPtr<FMP4BoxBase, ESPMode::ThreadSafe>& InParent, const FMP4BoxInfo& InBoxInfo)
 				: FMP4BoxBase(InParent, InBoxInfo)
 			{ }
 		};
@@ -586,7 +586,7 @@ namespace Electra
 		class FMP4BoxCOLR : public FMP4BoxBase
 		{
 		public:
-			static TSharedPtr<FMP4BoxBase> Create(const TWeakPtr<FMP4BoxBase>& InParent, const FMP4BoxInfo& InBoxInfo)
+			static TSharedPtr<FMP4BoxBase, ESPMode::ThreadSafe> Create(const TWeakPtr<FMP4BoxBase, ESPMode::ThreadSafe>& InParent, const FMP4BoxInfo& InBoxInfo)
 			{ return MakeShareable<FMP4BoxCOLR>(new FMP4BoxCOLR(InParent, InBoxInfo)); }
 			virtual ~FMP4BoxCOLR() = default;
 			enum class EColorType
@@ -606,7 +606,7 @@ namespace Electra
 			ELECTRABASE_API EColorType GetColorType();
 			ELECTRABASE_API const FColorNCLX& GetColorNCLX();
 		protected:
-			FMP4BoxCOLR(const TWeakPtr<FMP4BoxBase>& InParent, const FMP4BoxInfo& InBoxInfo)
+			FMP4BoxCOLR(const TWeakPtr<FMP4BoxBase, ESPMode::ThreadSafe>& InParent, const FMP4BoxInfo& InBoxInfo)
 				: FMP4BoxBase(InParent, InBoxInfo)
 			{ }
 			void ParseIfRequired();
@@ -622,12 +622,12 @@ namespace Electra
 		class FMP4BoxCLLI : public FMP4BoxBase
 		{
 		public:
-			static TSharedPtr<FMP4BoxBase> Create(const TWeakPtr<FMP4BoxBase>& InParent, const FMP4BoxInfo& InBoxInfo)
+			static TSharedPtr<FMP4BoxBase, ESPMode::ThreadSafe> Create(const TWeakPtr<FMP4BoxBase, ESPMode::ThreadSafe>& InParent, const FMP4BoxInfo& InBoxInfo)
 			{ return MakeShareable<FMP4BoxCLLI>(new FMP4BoxCLLI(InParent, InBoxInfo)); }
 			virtual ~FMP4BoxCLLI() = default;
 			//ELECTRABASE_API xxxx GetXXXX();
 		protected:
-			FMP4BoxCLLI(const TWeakPtr<FMP4BoxBase>& InParent, const FMP4BoxInfo& InBoxInfo)
+			FMP4BoxCLLI(const TWeakPtr<FMP4BoxBase, ESPMode::ThreadSafe>& InParent, const FMP4BoxInfo& InBoxInfo)
 				: FMP4BoxBase(InParent, InBoxInfo)
 			{ }
 			void ParseIfRequired();
@@ -643,12 +643,12 @@ namespace Electra
 		class FMP4BoxMDCV : public FMP4BoxBase
 		{
 		public:
-			static TSharedPtr<FMP4BoxBase> Create(const TWeakPtr<FMP4BoxBase>& InParent, const FMP4BoxInfo& InBoxInfo)
+			static TSharedPtr<FMP4BoxBase, ESPMode::ThreadSafe> Create(const TWeakPtr<FMP4BoxBase, ESPMode::ThreadSafe>& InParent, const FMP4BoxInfo& InBoxInfo)
 			{ return MakeShareable<FMP4BoxMDCV>(new FMP4BoxMDCV(InParent, InBoxInfo)); }
 			virtual ~FMP4BoxMDCV() = default;
 			//ELECTRABASE_API xxxx GetXXXX();
 		protected:
-			FMP4BoxMDCV(const TWeakPtr<FMP4BoxBase>& InParent, const FMP4BoxInfo& InBoxInfo)
+			FMP4BoxMDCV(const TWeakPtr<FMP4BoxBase, ESPMode::ThreadSafe>& InParent, const FMP4BoxInfo& InBoxInfo)
 				: FMP4BoxBase(InParent, InBoxInfo)
 			{ }
 			void ParseIfRequired();
@@ -664,7 +664,7 @@ namespace Electra
 		class FMP4BoxSTSD : public FMP4BoxBase
 		{
 		public:
-			static TSharedPtr<FMP4BoxBase> Create(const TWeakPtr<FMP4BoxBase>& InParent, const FMP4BoxInfo& InBoxInfo)
+			static TSharedPtr<FMP4BoxBase, ESPMode::ThreadSafe> Create(const TWeakPtr<FMP4BoxBase, ESPMode::ThreadSafe>& InParent, const FMP4BoxInfo& InBoxInfo)
 			{ return MakeShareable<FMP4BoxSTSD>(new FMP4BoxSTSD(InParent, InBoxInfo)); }
 			virtual ~FMP4BoxSTSD() = default;
 			bool IsLeafBox() const override { return false; }
@@ -674,7 +674,7 @@ namespace Electra
 			ELECTRABASE_API FMP4BoxSampleEntry::ESampleType GetSampleType();
 
 			template<typename T>
-			void GetSampleDescriptions(TArray<TSharedPtr<T>>& OutChildren)
+			void GetSampleDescriptions(TArray<TSharedPtr<T, ESPMode::ThreadSafe>>& OutChildren)
 			{
 				// Get the sample type managed by this box. This may implicitly trigger parsing!
 				const FMP4BoxSampleEntry::ESampleType SampleType = GetSampleType();
@@ -691,7 +691,7 @@ namespace Electra
 			}
 
 		protected:
-			FMP4BoxSTSD(const TWeakPtr<FMP4BoxBase>& InParent, const FMP4BoxInfo& InBoxInfo)
+			FMP4BoxSTSD(const TWeakPtr<FMP4BoxBase, ESPMode::ThreadSafe>& InParent, const FMP4BoxInfo& InBoxInfo)
 				: FMP4BoxBase(InParent, InBoxInfo)
 			{ }
 			void ParseIfRequired();
@@ -707,7 +707,7 @@ namespace Electra
 		class FMP4BoxSTTS : public FMP4BoxBase
 		{
 		public:
-			static TSharedPtr<FMP4BoxBase> Create(const TWeakPtr<FMP4BoxBase>& InParent, const FMP4BoxInfo& InBoxInfo)
+			static TSharedPtr<FMP4BoxBase, ESPMode::ThreadSafe> Create(const TWeakPtr<FMP4BoxBase, ESPMode::ThreadSafe>& InParent, const FMP4BoxInfo& InBoxInfo)
 			{ return MakeShareable<FMP4BoxSTTS>(new FMP4BoxSTTS(InParent, InBoxInfo)); }
 			virtual ~FMP4BoxSTTS() = default;
 			struct FEntry
@@ -719,7 +719,7 @@ namespace Electra
 			ELECTRABASE_API uint32 GetNumTotalSamples();
 			ELECTRABASE_API int64 GetTotalDuration();
 		protected:
-			FMP4BoxSTTS(const TWeakPtr<FMP4BoxBase>& InParent, const FMP4BoxInfo& InBoxInfo)
+			FMP4BoxSTTS(const TWeakPtr<FMP4BoxBase, ESPMode::ThreadSafe>& InParent, const FMP4BoxInfo& InBoxInfo)
 				: FMP4BoxBase(InParent, InBoxInfo)
 			{ }
 			void ParseIfRequired();
@@ -735,7 +735,7 @@ namespace Electra
 		class FMP4BoxCTTS : public FMP4BoxBase
 		{
 		public:
-			static TSharedPtr<FMP4BoxBase> Create(const TWeakPtr<FMP4BoxBase>& InParent, const FMP4BoxInfo& InBoxInfo)
+			static TSharedPtr<FMP4BoxBase, ESPMode::ThreadSafe> Create(const TWeakPtr<FMP4BoxBase, ESPMode::ThreadSafe>& InParent, const FMP4BoxInfo& InBoxInfo)
 			{ return MakeShareable<FMP4BoxCTTS>(new FMP4BoxCTTS(InParent, InBoxInfo)); }
 			virtual ~FMP4BoxCTTS() = default;
 			struct FEntry
@@ -748,7 +748,7 @@ namespace Electra
 			ELECTRABASE_API const TArray<FEntry>& GetEntries();
 			ELECTRABASE_API uint32 GetNumTotalSamples();
 		protected:
-			FMP4BoxCTTS(const TWeakPtr<FMP4BoxBase>& InParent, const FMP4BoxInfo& InBoxInfo)
+			FMP4BoxCTTS(const TWeakPtr<FMP4BoxBase, ESPMode::ThreadSafe>& InParent, const FMP4BoxInfo& InBoxInfo)
 				: FMP4BoxBase(InParent, InBoxInfo)
 			{ }
 			void ParseIfRequired();
@@ -764,11 +764,11 @@ namespace Electra
 		class FMP4BoxCSLG : public FMP4BoxBase
 		{
 		public:
-			static TSharedPtr<FMP4BoxBase> Create(const TWeakPtr<FMP4BoxBase>& InParent, const FMP4BoxInfo& InBoxInfo)
+			static TSharedPtr<FMP4BoxBase, ESPMode::ThreadSafe> Create(const TWeakPtr<FMP4BoxBase, ESPMode::ThreadSafe>& InParent, const FMP4BoxInfo& InBoxInfo)
 			{ return MakeShareable<FMP4BoxCSLG>(new FMP4BoxCSLG(InParent, InBoxInfo)); }
 			virtual ~FMP4BoxCSLG() = default;
 		protected:
-			FMP4BoxCSLG(const TWeakPtr<FMP4BoxBase>& InParent, const FMP4BoxInfo& InBoxInfo)
+			FMP4BoxCSLG(const TWeakPtr<FMP4BoxBase, ESPMode::ThreadSafe>& InParent, const FMP4BoxInfo& InBoxInfo)
 				: FMP4BoxBase(InParent, InBoxInfo)
 			{ }
 		};
@@ -781,12 +781,12 @@ namespace Electra
 		class FMP4BoxSTSS : public FMP4BoxBase
 		{
 		public:
-			static TSharedPtr<FMP4BoxBase> Create(const TWeakPtr<FMP4BoxBase>& InParent, const FMP4BoxInfo& InBoxInfo)
+			static TSharedPtr<FMP4BoxBase, ESPMode::ThreadSafe> Create(const TWeakPtr<FMP4BoxBase, ESPMode::ThreadSafe>& InParent, const FMP4BoxInfo& InBoxInfo)
 			{ return MakeShareable<FMP4BoxSTSS>(new FMP4BoxSTSS(InParent, InBoxInfo)); }
 			virtual ~FMP4BoxSTSS() = default;
 			ELECTRABASE_API const TArray<uint32>& GetEntries();
 		protected:
-			FMP4BoxSTSS(const TWeakPtr<FMP4BoxBase>& InParent, const FMP4BoxInfo& InBoxInfo)
+			FMP4BoxSTSS(const TWeakPtr<FMP4BoxBase, ESPMode::ThreadSafe>& InParent, const FMP4BoxInfo& InBoxInfo)
 				: FMP4BoxBase(InParent, InBoxInfo)
 			{ }
 			void ParseIfRequired();
@@ -802,12 +802,12 @@ namespace Electra
 		class FMP4BoxSDTP : public FMP4BoxBase
 		{
 		public:
-			static TSharedPtr<FMP4BoxBase> Create(const TWeakPtr<FMP4BoxBase>& InParent, const FMP4BoxInfo& InBoxInfo)
+			static TSharedPtr<FMP4BoxBase, ESPMode::ThreadSafe> Create(const TWeakPtr<FMP4BoxBase, ESPMode::ThreadSafe>& InParent, const FMP4BoxInfo& InBoxInfo)
 			{ return MakeShareable<FMP4BoxSDTP>(new FMP4BoxSDTP(InParent, InBoxInfo)); }
 			virtual ~FMP4BoxSDTP() = default;
 			ELECTRABASE_API TConstArrayView<uint8> GetEntries();
 		protected:
-			FMP4BoxSDTP(const TWeakPtr<FMP4BoxBase>& InParent, const FMP4BoxInfo& InBoxInfo)
+			FMP4BoxSDTP(const TWeakPtr<FMP4BoxBase, ESPMode::ThreadSafe>& InParent, const FMP4BoxInfo& InBoxInfo)
 				: FMP4BoxBase(InParent, InBoxInfo)
 			{ }
 			void ParseIfRequired();
@@ -823,7 +823,7 @@ namespace Electra
 		class FMP4BoxSTSC : public FMP4BoxBase
 		{
 		public:
-			static TSharedPtr<FMP4BoxBase> Create(const TWeakPtr<FMP4BoxBase>& InParent, const FMP4BoxInfo& InBoxInfo)
+			static TSharedPtr<FMP4BoxBase, ESPMode::ThreadSafe> Create(const TWeakPtr<FMP4BoxBase, ESPMode::ThreadSafe>& InParent, const FMP4BoxInfo& InBoxInfo)
 			{ return MakeShareable<FMP4BoxSTSC>(new FMP4BoxSTSC(InParent, InBoxInfo)); }
 			virtual ~FMP4BoxSTSC() = default;
 			struct FEntry
@@ -834,7 +834,7 @@ namespace Electra
 			};
 			ELECTRABASE_API const TArray<FEntry>& GetEntries();
 		protected:
-			FMP4BoxSTSC(const TWeakPtr<FMP4BoxBase>& InParent, const FMP4BoxInfo& InBoxInfo)
+			FMP4BoxSTSC(const TWeakPtr<FMP4BoxBase, ESPMode::ThreadSafe>& InParent, const FMP4BoxInfo& InBoxInfo)
 				: FMP4BoxBase(InParent, InBoxInfo)
 			{ }
 			void ParseIfRequired();
@@ -850,13 +850,13 @@ namespace Electra
 		class FMP4BoxSTSZ : public FMP4BoxBase
 		{
 		public:
-			static TSharedPtr<FMP4BoxBase> Create(const TWeakPtr<FMP4BoxBase>& InParent, const FMP4BoxInfo& InBoxInfo)
+			static TSharedPtr<FMP4BoxBase, ESPMode::ThreadSafe> Create(const TWeakPtr<FMP4BoxBase, ESPMode::ThreadSafe>& InParent, const FMP4BoxInfo& InBoxInfo)
 			{ return MakeShareable<FMP4BoxSTSZ>(new FMP4BoxSTSZ(InParent, InBoxInfo)); }
 			virtual ~FMP4BoxSTSZ() = default;
 			ELECTRABASE_API uint32 GetNumberOfSamples();
 			ELECTRABASE_API uint32 GetSizeOfSample(uint32 InIndex);
 		protected:
-			FMP4BoxSTSZ(const TWeakPtr<FMP4BoxBase>& InParent, const FMP4BoxInfo& InBoxInfo)
+			FMP4BoxSTSZ(const TWeakPtr<FMP4BoxBase, ESPMode::ThreadSafe>& InParent, const FMP4BoxInfo& InBoxInfo)
 				: FMP4BoxBase(InParent, InBoxInfo)
 			{ }
 			void ParseIfRequired();
@@ -872,7 +872,7 @@ namespace Electra
 		class FMP4BoxSTCO : public FMP4BoxBase
 		{
 		public:
-			static TSharedPtr<FMP4BoxBase> Create(const TWeakPtr<FMP4BoxBase>& InParent, const FMP4BoxInfo& InBoxInfo)
+			static TSharedPtr<FMP4BoxBase, ESPMode::ThreadSafe> Create(const TWeakPtr<FMP4BoxBase, ESPMode::ThreadSafe>& InParent, const FMP4BoxInfo& InBoxInfo)
 			{ return MakeShareable<FMP4BoxSTCO>(new FMP4BoxSTCO(InParent, InBoxInfo)); }
 			virtual ~FMP4BoxSTCO() = default;
 			uint32 GetType() const override
@@ -885,7 +885,7 @@ namespace Electra
 			ELECTRABASE_API uint32 GetNumberOfEntries();
 			ELECTRABASE_API uint64 GetChunkOffset(uint32 InIndex);
 		protected:
-			FMP4BoxSTCO(const TWeakPtr<FMP4BoxBase>& InParent, const FMP4BoxInfo& InBoxInfo)
+			FMP4BoxSTCO(const TWeakPtr<FMP4BoxBase, ESPMode::ThreadSafe>& InParent, const FMP4BoxInfo& InBoxInfo)
 				: FMP4BoxBase(InParent, InBoxInfo)
 			{ }
 			void ParseIfRequired();
@@ -901,12 +901,12 @@ namespace Electra
 		class FMP4BoxSAIZ : public FMP4BoxBase
 		{
 		public:
-			static TSharedPtr<FMP4BoxBase> Create(const TWeakPtr<FMP4BoxBase>& InParent, const FMP4BoxInfo& InBoxInfo)
+			static TSharedPtr<FMP4BoxBase, ESPMode::ThreadSafe> Create(const TWeakPtr<FMP4BoxBase, ESPMode::ThreadSafe>& InParent, const FMP4BoxInfo& InBoxInfo)
 			{ return MakeShareable<FMP4BoxSAIZ>(new FMP4BoxSAIZ(InParent, InBoxInfo)); }
 			virtual ~FMP4BoxSAIZ() = default;
 			ELECTRABASE_API void Test();
 		protected:
-			FMP4BoxSAIZ(const TWeakPtr<FMP4BoxBase>& InParent, const FMP4BoxInfo& InBoxInfo)
+			FMP4BoxSAIZ(const TWeakPtr<FMP4BoxBase, ESPMode::ThreadSafe>& InParent, const FMP4BoxInfo& InBoxInfo)
 				: FMP4BoxBase(InParent, InBoxInfo)
 			{ }
 			void ParseIfRequired();
@@ -922,12 +922,12 @@ namespace Electra
 		class FMP4BoxSAIO : public FMP4BoxBase
 		{
 		public:
-			static TSharedPtr<FMP4BoxBase> Create(const TWeakPtr<FMP4BoxBase>& InParent, const FMP4BoxInfo& InBoxInfo)
+			static TSharedPtr<FMP4BoxBase, ESPMode::ThreadSafe> Create(const TWeakPtr<FMP4BoxBase, ESPMode::ThreadSafe>& InParent, const FMP4BoxInfo& InBoxInfo)
 			{ return MakeShareable<FMP4BoxSAIO>(new FMP4BoxSAIO(InParent, InBoxInfo)); }
 			virtual ~FMP4BoxSAIO() = default;
 			ELECTRABASE_API void Test();
 		protected:
-			FMP4BoxSAIO(const TWeakPtr<FMP4BoxBase>& InParent, const FMP4BoxInfo& InBoxInfo)
+			FMP4BoxSAIO(const TWeakPtr<FMP4BoxBase, ESPMode::ThreadSafe>& InParent, const FMP4BoxInfo& InBoxInfo)
 				: FMP4BoxBase(InParent, InBoxInfo)
 			{ }
 			void ParseIfRequired();
@@ -943,7 +943,7 @@ namespace Electra
 		class FMP4BoxSGPD : public FMP4BoxBase
 		{
 		public:
-			static TSharedPtr<FMP4BoxBase> Create(const TWeakPtr<FMP4BoxBase>& InParent, const FMP4BoxInfo& InBoxInfo)
+			static TSharedPtr<FMP4BoxBase, ESPMode::ThreadSafe> Create(const TWeakPtr<FMP4BoxBase, ESPMode::ThreadSafe>& InParent, const FMP4BoxInfo& InBoxInfo)
 			{ return MakeShareable<FMP4BoxSGPD>(new FMP4BoxSGPD(InParent, InBoxInfo)); }
 			virtual ~FMP4BoxSGPD() = default;
 
@@ -951,7 +951,7 @@ namespace Electra
 			ELECTRABASE_API const TArray<TConstArrayView<uint8>>& GetGroupDescriptionEntries();
 			ELECTRABASE_API uint32 GetDefaultGroupDescriptionIndex();
 		protected:
-			FMP4BoxSGPD(const TWeakPtr<FMP4BoxBase>& InParent, const FMP4BoxInfo& InBoxInfo)
+			FMP4BoxSGPD(const TWeakPtr<FMP4BoxBase, ESPMode::ThreadSafe>& InParent, const FMP4BoxInfo& InBoxInfo)
 				: FMP4BoxBase(InParent, InBoxInfo)
 			{ }
 			void ParseIfRequired();
@@ -967,7 +967,7 @@ namespace Electra
 		class FMP4BoxSBGP : public FMP4BoxBase
 		{
 		public:
-			static TSharedPtr<FMP4BoxBase> Create(const TWeakPtr<FMP4BoxBase>& InParent, const FMP4BoxInfo& InBoxInfo)
+			static TSharedPtr<FMP4BoxBase, ESPMode::ThreadSafe> Create(const TWeakPtr<FMP4BoxBase, ESPMode::ThreadSafe>& InParent, const FMP4BoxInfo& InBoxInfo)
 			{ return MakeShareable<FMP4BoxSBGP>(new FMP4BoxSBGP(InParent, InBoxInfo)); }
 			virtual ~FMP4BoxSBGP() = default;
 
@@ -982,7 +982,7 @@ namespace Electra
 			ELECTRABASE_API const TArray<FEntry>& GetEntries();
 			ELECTRABASE_API uint32 GetNumTotalSamples();
 		protected:
-			FMP4BoxSBGP(const TWeakPtr<FMP4BoxBase>& InParent, const FMP4BoxInfo& InBoxInfo)
+			FMP4BoxSBGP(const TWeakPtr<FMP4BoxBase, ESPMode::ThreadSafe>& InParent, const FMP4BoxInfo& InBoxInfo)
 				: FMP4BoxBase(InParent, InBoxInfo)
 			{ }
 			void ParseIfRequired();
@@ -998,12 +998,12 @@ namespace Electra
 		class FMP4BoxMVEX : public FMP4BoxBase
 		{
 		public:
-			static TSharedPtr<FMP4BoxBase> Create(const TWeakPtr<FMP4BoxBase>& InParent, const FMP4BoxInfo& InBoxInfo)
+			static TSharedPtr<FMP4BoxBase, ESPMode::ThreadSafe> Create(const TWeakPtr<FMP4BoxBase, ESPMode::ThreadSafe>& InParent, const FMP4BoxInfo& InBoxInfo)
 			{ return MakeShareable<FMP4BoxMVEX>(new FMP4BoxMVEX(InParent, InBoxInfo)); }
 			virtual ~FMP4BoxMVEX() = default;
 			bool IsLeafBox() const override { return false; }
 		protected:
-			FMP4BoxMVEX(const TWeakPtr<FMP4BoxBase>& InParent, const FMP4BoxInfo& InBoxInfo)
+			FMP4BoxMVEX(const TWeakPtr<FMP4BoxBase, ESPMode::ThreadSafe>& InParent, const FMP4BoxInfo& InBoxInfo)
 				: FMP4BoxBase(InParent, InBoxInfo)
 			{ }
 		};
@@ -1016,7 +1016,7 @@ namespace Electra
 		class FMP4BoxMEHD : public FMP4BoxBase
 		{
 		public:
-			static TSharedPtr<FMP4BoxBase> Create(const TWeakPtr<FMP4BoxBase>& InParent, const FMP4BoxInfo& InBoxInfo)
+			static TSharedPtr<FMP4BoxBase, ESPMode::ThreadSafe> Create(const TWeakPtr<FMP4BoxBase, ESPMode::ThreadSafe>& InParent, const FMP4BoxInfo& InBoxInfo)
 			{ return MakeShareable<FMP4BoxMEHD>(new FMP4BoxMEHD(InParent, InBoxInfo)); }
 			virtual ~FMP4BoxMEHD() = default;
 			/**
@@ -1025,7 +1025,7 @@ namespace Electra
 			 */
 			ELECTRABASE_API uint64 GetFragmentDuration();
 		protected:
-			FMP4BoxMEHD(const TWeakPtr<FMP4BoxBase>& InParent, const FMP4BoxInfo& InBoxInfo)
+			FMP4BoxMEHD(const TWeakPtr<FMP4BoxBase, ESPMode::ThreadSafe>& InParent, const FMP4BoxInfo& InBoxInfo)
 				: FMP4BoxBase(InParent, InBoxInfo)
 			{ }
 			void ParseIfRequired();
@@ -1041,7 +1041,7 @@ namespace Electra
 		class FMP4BoxTREX : public FMP4BoxBase
 		{
 		public:
-			static TSharedPtr<FMP4BoxBase> Create(const TWeakPtr<FMP4BoxBase>& InParent, const FMP4BoxInfo& InBoxInfo)
+			static TSharedPtr<FMP4BoxBase, ESPMode::ThreadSafe> Create(const TWeakPtr<FMP4BoxBase, ESPMode::ThreadSafe>& InParent, const FMP4BoxInfo& InBoxInfo)
 			{ return MakeShareable<FMP4BoxTREX>(new FMP4BoxTREX(InParent, InBoxInfo)); }
 			virtual ~FMP4BoxTREX() = default;
 			ELECTRABASE_API uint32 GetTrackID();
@@ -1050,7 +1050,7 @@ namespace Electra
 			ELECTRABASE_API uint32 GetDefaultSampleSize();
 			ELECTRABASE_API uint32 GetDefaultSampleFlags();
 		protected:
-			FMP4BoxTREX(const TWeakPtr<FMP4BoxBase>& InParent, const FMP4BoxInfo& InBoxInfo)
+			FMP4BoxTREX(const TWeakPtr<FMP4BoxBase, ESPMode::ThreadSafe>& InParent, const FMP4BoxInfo& InBoxInfo)
 				: FMP4BoxBase(InParent, InBoxInfo)
 			{ }
 			void ParseIfRequired();
@@ -1066,11 +1066,11 @@ namespace Electra
 		class FMP4BoxLEVA : public FMP4BoxBase
 		{
 		public:
-			static TSharedPtr<FMP4BoxBase> Create(const TWeakPtr<FMP4BoxBase>& InParent, const FMP4BoxInfo& InBoxInfo)
+			static TSharedPtr<FMP4BoxBase, ESPMode::ThreadSafe> Create(const TWeakPtr<FMP4BoxBase, ESPMode::ThreadSafe>& InParent, const FMP4BoxInfo& InBoxInfo)
 			{ return MakeShareable<FMP4BoxLEVA>(new FMP4BoxLEVA(InParent, InBoxInfo)); }
 			virtual ~FMP4BoxLEVA() = default;
 		protected:
-			FMP4BoxLEVA(const TWeakPtr<FMP4BoxBase>& InParent, const FMP4BoxInfo& InBoxInfo)
+			FMP4BoxLEVA(const TWeakPtr<FMP4BoxBase, ESPMode::ThreadSafe>& InParent, const FMP4BoxInfo& InBoxInfo)
 				: FMP4BoxBase(InParent, InBoxInfo)
 			{ }
 		};
@@ -1083,12 +1083,12 @@ namespace Electra
 		class FMP4BoxMOOF : public FMP4BoxBase
 		{
 		public:
-			static TSharedPtr<FMP4BoxBase> Create(const TWeakPtr<FMP4BoxBase>& InParent, const FMP4BoxInfo& InBoxInfo)
+			static TSharedPtr<FMP4BoxBase, ESPMode::ThreadSafe> Create(const TWeakPtr<FMP4BoxBase, ESPMode::ThreadSafe>& InParent, const FMP4BoxInfo& InBoxInfo)
 			{ return MakeShareable<FMP4BoxMOOF>(new FMP4BoxMOOF(InParent, InBoxInfo)); }
 			virtual ~FMP4BoxMOOF() = default;
 			bool IsLeafBox() const override { return false; }
 		protected:
-			FMP4BoxMOOF(const TWeakPtr<FMP4BoxBase>& InParent, const FMP4BoxInfo& InBoxInfo)
+			FMP4BoxMOOF(const TWeakPtr<FMP4BoxBase, ESPMode::ThreadSafe>& InParent, const FMP4BoxInfo& InBoxInfo)
 				: FMP4BoxBase(InParent, InBoxInfo)
 			{ }
 		};
@@ -1101,12 +1101,12 @@ namespace Electra
 		class FMP4BoxMFHD : public FMP4BoxBase
 		{
 		public:
-			static TSharedPtr<FMP4BoxBase> Create(const TWeakPtr<FMP4BoxBase>& InParent, const FMP4BoxInfo& InBoxInfo)
+			static TSharedPtr<FMP4BoxBase, ESPMode::ThreadSafe> Create(const TWeakPtr<FMP4BoxBase, ESPMode::ThreadSafe>& InParent, const FMP4BoxInfo& InBoxInfo)
 			{ return MakeShareable<FMP4BoxMFHD>(new FMP4BoxMFHD(InParent, InBoxInfo)); }
 			virtual ~FMP4BoxMFHD() = default;
 			ELECTRABASE_API uint32 GetSequenceNumber();
 		protected:
-			FMP4BoxMFHD(const TWeakPtr<FMP4BoxBase>& InParent, const FMP4BoxInfo& InBoxInfo)
+			FMP4BoxMFHD(const TWeakPtr<FMP4BoxBase, ESPMode::ThreadSafe>& InParent, const FMP4BoxInfo& InBoxInfo)
 				: FMP4BoxBase(InParent, InBoxInfo)
 			{ }
 			void ParseIfRequired();
@@ -1122,12 +1122,12 @@ namespace Electra
 		class FMP4BoxTRAF : public FMP4BoxBase
 		{
 		public:
-			static TSharedPtr<FMP4BoxBase> Create(const TWeakPtr<FMP4BoxBase>& InParent, const FMP4BoxInfo& InBoxInfo)
+			static TSharedPtr<FMP4BoxBase, ESPMode::ThreadSafe> Create(const TWeakPtr<FMP4BoxBase, ESPMode::ThreadSafe>& InParent, const FMP4BoxInfo& InBoxInfo)
 			{ return MakeShareable<FMP4BoxTRAF>(new FMP4BoxTRAF(InParent, InBoxInfo)); }
 			virtual ~FMP4BoxTRAF() = default;
 			bool IsLeafBox() const override { return false; }
 		protected:
-			FMP4BoxTRAF(const TWeakPtr<FMP4BoxBase>& InParent, const FMP4BoxInfo& InBoxInfo)
+			FMP4BoxTRAF(const TWeakPtr<FMP4BoxBase, ESPMode::ThreadSafe>& InParent, const FMP4BoxInfo& InBoxInfo)
 				: FMP4BoxBase(InParent, InBoxInfo)
 			{ }
 		};
@@ -1140,7 +1140,7 @@ namespace Electra
 		class FMP4BoxTFHD : public FMP4BoxBase
 		{
 		public:
-			static TSharedPtr<FMP4BoxBase> Create(const TWeakPtr<FMP4BoxBase>& InParent, const FMP4BoxInfo& InBoxInfo)
+			static TSharedPtr<FMP4BoxBase, ESPMode::ThreadSafe> Create(const TWeakPtr<FMP4BoxBase, ESPMode::ThreadSafe>& InParent, const FMP4BoxInfo& InBoxInfo)
 			{ return MakeShareable<FMP4BoxTFHD>(new FMP4BoxTFHD(InParent, InBoxInfo)); }
 			virtual ~FMP4BoxTFHD() = default;
 
@@ -1158,7 +1158,7 @@ namespace Electra
 			ELECTRABASE_API bool IsDurationEmpty();
 			ELECTRABASE_API bool IsMoofDefaultBase();
 		protected:
-			FMP4BoxTFHD(const TWeakPtr<FMP4BoxBase>& InParent, const FMP4BoxInfo& InBoxInfo)
+			FMP4BoxTFHD(const TWeakPtr<FMP4BoxBase, ESPMode::ThreadSafe>& InParent, const FMP4BoxInfo& InBoxInfo)
 				: FMP4BoxBase(InParent, InBoxInfo)
 			{ }
 			void ParseIfRequired();
@@ -1174,7 +1174,7 @@ namespace Electra
 		class FMP4BoxTRUN : public FMP4BoxBase
 		{
 		public:
-			static TSharedPtr<FMP4BoxBase> Create(const TWeakPtr<FMP4BoxBase>& InParent, const FMP4BoxInfo& InBoxInfo)
+			static TSharedPtr<FMP4BoxBase, ESPMode::ThreadSafe> Create(const TWeakPtr<FMP4BoxBase, ESPMode::ThreadSafe>& InParent, const FMP4BoxInfo& InBoxInfo)
 			{ return MakeShareable<FMP4BoxTRUN>(new FMP4BoxTRUN(InParent, InBoxInfo)); }
 			virtual ~FMP4BoxTRUN() = default;
 			ELECTRABASE_API uint32 GetNumberOfSamples();
@@ -1191,7 +1191,7 @@ namespace Electra
 			ELECTRABASE_API bool HasSampleCompositionTimeOffsets();
 			ELECTRABASE_API const TArray<int32>& GetSampleCompositionTimeOffsets();
 		protected:
-			FMP4BoxTRUN(const TWeakPtr<FMP4BoxBase>& InParent, const FMP4BoxInfo& InBoxInfo)
+			FMP4BoxTRUN(const TWeakPtr<FMP4BoxBase, ESPMode::ThreadSafe>& InParent, const FMP4BoxInfo& InBoxInfo)
 				: FMP4BoxBase(InParent, InBoxInfo)
 			{ }
 			void ParseIfRequired();
@@ -1207,12 +1207,12 @@ namespace Electra
 		class FMP4BoxMFRA : public FMP4BoxBase
 		{
 		public:
-			static TSharedPtr<FMP4BoxBase> Create(const TWeakPtr<FMP4BoxBase>& InParent, const FMP4BoxInfo& InBoxInfo)
+			static TSharedPtr<FMP4BoxBase, ESPMode::ThreadSafe> Create(const TWeakPtr<FMP4BoxBase, ESPMode::ThreadSafe>& InParent, const FMP4BoxInfo& InBoxInfo)
 			{ return MakeShareable<FMP4BoxMFRA>(new FMP4BoxMFRA(InParent, InBoxInfo)); }
 			virtual ~FMP4BoxMFRA() = default;
 			bool IsLeafBox() const override { return false; }
 		protected:
-			FMP4BoxMFRA(const TWeakPtr<FMP4BoxBase>& InParent, const FMP4BoxInfo& InBoxInfo)
+			FMP4BoxMFRA(const TWeakPtr<FMP4BoxBase, ESPMode::ThreadSafe>& InParent, const FMP4BoxInfo& InBoxInfo)
 				: FMP4BoxBase(InParent, InBoxInfo)
 			{ }
 		};
@@ -1225,11 +1225,11 @@ namespace Electra
 		class FMP4BoxTFRA : public FMP4BoxBase
 		{
 		public:
-			static TSharedPtr<FMP4BoxBase> Create(const TWeakPtr<FMP4BoxBase>& InParent, const FMP4BoxInfo& InBoxInfo)
+			static TSharedPtr<FMP4BoxBase, ESPMode::ThreadSafe> Create(const TWeakPtr<FMP4BoxBase, ESPMode::ThreadSafe>& InParent, const FMP4BoxInfo& InBoxInfo)
 			{ return MakeShareable<FMP4BoxTFRA>(new FMP4BoxTFRA(InParent, InBoxInfo)); }
 			virtual ~FMP4BoxTFRA() = default;
 		protected:
-			FMP4BoxTFRA(const TWeakPtr<FMP4BoxBase>& InParent, const FMP4BoxInfo& InBoxInfo)
+			FMP4BoxTFRA(const TWeakPtr<FMP4BoxBase, ESPMode::ThreadSafe>& InParent, const FMP4BoxInfo& InBoxInfo)
 				: FMP4BoxBase(InParent, InBoxInfo)
 			{ }
 		};
@@ -1242,11 +1242,11 @@ namespace Electra
 		class FMP4BoxMFRO : public FMP4BoxBase
 		{
 		public:
-			static TSharedPtr<FMP4BoxBase> Create(const TWeakPtr<FMP4BoxBase>& InParent, const FMP4BoxInfo& InBoxInfo)
+			static TSharedPtr<FMP4BoxBase, ESPMode::ThreadSafe> Create(const TWeakPtr<FMP4BoxBase, ESPMode::ThreadSafe>& InParent, const FMP4BoxInfo& InBoxInfo)
 			{ return MakeShareable<FMP4BoxMFRO>(new FMP4BoxMFRO(InParent, InBoxInfo)); }
 			virtual ~FMP4BoxMFRO() = default;
 		protected:
-			FMP4BoxMFRO(const TWeakPtr<FMP4BoxBase>& InParent, const FMP4BoxInfo& InBoxInfo)
+			FMP4BoxMFRO(const TWeakPtr<FMP4BoxBase, ESPMode::ThreadSafe>& InParent, const FMP4BoxInfo& InBoxInfo)
 				: FMP4BoxBase(InParent, InBoxInfo)
 			{ }
 		};
@@ -1259,12 +1259,12 @@ namespace Electra
 		class FMP4BoxTFDT : public FMP4BoxBase
 		{
 		public:
-			static TSharedPtr<FMP4BoxBase> Create(const TWeakPtr<FMP4BoxBase>& InParent, const FMP4BoxInfo& InBoxInfo)
+			static TSharedPtr<FMP4BoxBase, ESPMode::ThreadSafe> Create(const TWeakPtr<FMP4BoxBase, ESPMode::ThreadSafe>& InParent, const FMP4BoxInfo& InBoxInfo)
 			{ return MakeShareable<FMP4BoxTFDT>(new FMP4BoxTFDT(InParent, InBoxInfo)); }
 			virtual ~FMP4BoxTFDT() = default;
 			ELECTRABASE_API uint64 GetBaseMediaDecodeTime();
 		protected:
-			FMP4BoxTFDT(const TWeakPtr<FMP4BoxBase>& InParent, const FMP4BoxInfo& InBoxInfo)
+			FMP4BoxTFDT(const TWeakPtr<FMP4BoxBase, ESPMode::ThreadSafe>& InParent, const FMP4BoxInfo& InBoxInfo)
 				: FMP4BoxBase(InParent, InBoxInfo)
 			{ }
 			void ParseIfRequired();
@@ -1280,7 +1280,7 @@ namespace Electra
 		class FMP4BoxSIDX : public FMP4BoxBase
 		{
 		public:
-			static TSharedPtr<FMP4BoxBase> Create(const TWeakPtr<FMP4BoxBase>& InParent, const FMP4BoxInfo& InBoxInfo)
+			static TSharedPtr<FMP4BoxBase, ESPMode::ThreadSafe> Create(const TWeakPtr<FMP4BoxBase, ESPMode::ThreadSafe>& InParent, const FMP4BoxInfo& InBoxInfo)
 			{ return MakeShareable<FMP4BoxSIDX>(new FMP4BoxSIDX(InParent, InBoxInfo)); }
 			virtual ~FMP4BoxSIDX() = default;
 			struct FEntry
@@ -1299,7 +1299,7 @@ namespace Electra
 			ELECTRABASE_API uint64 GetFirstOffset();
 			ELECTRABASE_API const FEntryList& GetEntries();
 		protected:
-			FMP4BoxSIDX(const TWeakPtr<FMP4BoxBase>& InParent, const FMP4BoxInfo& InBoxInfo)
+			FMP4BoxSIDX(const TWeakPtr<FMP4BoxBase, ESPMode::ThreadSafe>& InParent, const FMP4BoxInfo& InBoxInfo)
 				: FMP4BoxBase(InParent, InBoxInfo)
 			{ }
 			void ParseIfRequired();
@@ -1315,11 +1315,11 @@ namespace Electra
 		class FMP4BoxSSIX : public FMP4BoxBase
 		{
 		public:
-			static TSharedPtr<FMP4BoxBase> Create(const TWeakPtr<FMP4BoxBase>& InParent, const FMP4BoxInfo& InBoxInfo)
+			static TSharedPtr<FMP4BoxBase, ESPMode::ThreadSafe> Create(const TWeakPtr<FMP4BoxBase, ESPMode::ThreadSafe>& InParent, const FMP4BoxInfo& InBoxInfo)
 			{ return MakeShareable<FMP4BoxSSIX>(new FMP4BoxSSIX(InParent, InBoxInfo)); }
 			virtual ~FMP4BoxSSIX() = default;
 		protected:
-			FMP4BoxSSIX(const TWeakPtr<FMP4BoxBase>& InParent, const FMP4BoxInfo& InBoxInfo)
+			FMP4BoxSSIX(const TWeakPtr<FMP4BoxBase, ESPMode::ThreadSafe>& InParent, const FMP4BoxInfo& InBoxInfo)
 				: FMP4BoxBase(InParent, InBoxInfo)
 			{ }
 		};
@@ -1332,11 +1332,11 @@ namespace Electra
 		class FMP4BoxVMHD : public FMP4BoxBase
 		{
 		public:
-			static TSharedPtr<FMP4BoxBase> Create(const TWeakPtr<FMP4BoxBase>& InParent, const FMP4BoxInfo& InBoxInfo)
+			static TSharedPtr<FMP4BoxBase, ESPMode::ThreadSafe> Create(const TWeakPtr<FMP4BoxBase, ESPMode::ThreadSafe>& InParent, const FMP4BoxInfo& InBoxInfo)
 			{ return MakeShareable<FMP4BoxVMHD>(new FMP4BoxVMHD(InParent, InBoxInfo)); }
 			virtual ~FMP4BoxVMHD() = default;
 		protected:
-			FMP4BoxVMHD(const TWeakPtr<FMP4BoxBase>& InParent, const FMP4BoxInfo& InBoxInfo)
+			FMP4BoxVMHD(const TWeakPtr<FMP4BoxBase, ESPMode::ThreadSafe>& InParent, const FMP4BoxInfo& InBoxInfo)
 				: FMP4BoxBase(InParent, InBoxInfo)
 			{ }
 		};
@@ -1349,11 +1349,11 @@ namespace Electra
 		class FMP4BoxSMHD : public FMP4BoxBase
 		{
 		public:
-			static TSharedPtr<FMP4BoxBase> Create(const TWeakPtr<FMP4BoxBase>& InParent, const FMP4BoxInfo& InBoxInfo)
+			static TSharedPtr<FMP4BoxBase, ESPMode::ThreadSafe> Create(const TWeakPtr<FMP4BoxBase, ESPMode::ThreadSafe>& InParent, const FMP4BoxInfo& InBoxInfo)
 			{ return MakeShareable<FMP4BoxSMHD>(new FMP4BoxSMHD(InParent, InBoxInfo)); }
 			virtual ~FMP4BoxSMHD() = default;
 		protected:
-			FMP4BoxSMHD(const TWeakPtr<FMP4BoxBase>& InParent, const FMP4BoxInfo& InBoxInfo)
+			FMP4BoxSMHD(const TWeakPtr<FMP4BoxBase, ESPMode::ThreadSafe>& InParent, const FMP4BoxInfo& InBoxInfo)
 				: FMP4BoxBase(InParent, InBoxInfo)
 			{ }
 		};
@@ -1366,11 +1366,11 @@ namespace Electra
 		class FMP4BoxNMHD : public FMP4BoxBase
 		{
 		public:
-			static TSharedPtr<FMP4BoxBase> Create(const TWeakPtr<FMP4BoxBase>& InParent, const FMP4BoxInfo& InBoxInfo)
+			static TSharedPtr<FMP4BoxBase, ESPMode::ThreadSafe> Create(const TWeakPtr<FMP4BoxBase, ESPMode::ThreadSafe>& InParent, const FMP4BoxInfo& InBoxInfo)
 			{ return MakeShareable<FMP4BoxNMHD>(new FMP4BoxNMHD(InParent, InBoxInfo)); }
 			virtual ~FMP4BoxNMHD() = default;
 		protected:
-			FMP4BoxNMHD(const TWeakPtr<FMP4BoxBase>& InParent, const FMP4BoxInfo& InBoxInfo)
+			FMP4BoxNMHD(const TWeakPtr<FMP4BoxBase, ESPMode::ThreadSafe>& InParent, const FMP4BoxInfo& InBoxInfo)
 				: FMP4BoxBase(InParent, InBoxInfo)
 			{ }
 		};
@@ -1383,11 +1383,11 @@ namespace Electra
 		class FMP4BoxGMHD : public FMP4BoxBase
 		{
 		public:
-			static TSharedPtr<FMP4BoxBase> Create(const TWeakPtr<FMP4BoxBase>& InParent, const FMP4BoxInfo& InBoxInfo)
+			static TSharedPtr<FMP4BoxBase, ESPMode::ThreadSafe> Create(const TWeakPtr<FMP4BoxBase, ESPMode::ThreadSafe>& InParent, const FMP4BoxInfo& InBoxInfo)
 			{ return MakeShareable<FMP4BoxGMHD>(new FMP4BoxGMHD(InParent, InBoxInfo)); }
 			virtual ~FMP4BoxGMHD() = default;
 		protected:
-			FMP4BoxGMHD(const TWeakPtr<FMP4BoxBase>& InParent, const FMP4BoxInfo& InBoxInfo)
+			FMP4BoxGMHD(const TWeakPtr<FMP4BoxBase, ESPMode::ThreadSafe>& InParent, const FMP4BoxInfo& InBoxInfo)
 				: FMP4BoxBase(InParent, InBoxInfo)
 			{ }
 		};
@@ -1400,12 +1400,12 @@ namespace Electra
 		class FMP4BoxUDTA : public FMP4BoxBase
 		{
 		public:
-			static TSharedPtr<FMP4BoxBase> Create(const TWeakPtr<FMP4BoxBase>& InParent, const FMP4BoxInfo& InBoxInfo)
+			static TSharedPtr<FMP4BoxBase, ESPMode::ThreadSafe> Create(const TWeakPtr<FMP4BoxBase, ESPMode::ThreadSafe>& InParent, const FMP4BoxInfo& InBoxInfo)
 			{ return MakeShareable<FMP4BoxUDTA>(new FMP4BoxUDTA(InParent, InBoxInfo)); }
 			virtual ~FMP4BoxUDTA() = default;
 			bool IsLeafBox() const override { return false; }
 		protected:
-			FMP4BoxUDTA(const TWeakPtr<FMP4BoxBase>& InParent, const FMP4BoxInfo& InBoxInfo)
+			FMP4BoxUDTA(const TWeakPtr<FMP4BoxBase, ESPMode::ThreadSafe>& InParent, const FMP4BoxInfo& InBoxInfo)
 				: FMP4BoxBase(InParent, InBoxInfo)
 			{ }
 		};
@@ -1418,13 +1418,13 @@ namespace Electra
 		class FMP4BoxMETA : public FMP4BoxBase
 		{
 		public:
-			static TSharedPtr<FMP4BoxBase> Create(const TWeakPtr<FMP4BoxBase>& InParent, const FMP4BoxInfo& InBoxInfo)
+			static TSharedPtr<FMP4BoxBase, ESPMode::ThreadSafe> Create(const TWeakPtr<FMP4BoxBase, ESPMode::ThreadSafe>& InParent, const FMP4BoxInfo& InBoxInfo)
 			{ return MakeShareable<FMP4BoxMETA>(new FMP4BoxMETA(InParent, InBoxInfo)); }
 			virtual ~FMP4BoxMETA() = default;
 			bool IsLeafBox() const override { return false; }
 			bool IsListOfEntries() const override { return true; }
 		protected:
-			FMP4BoxMETA(const TWeakPtr<FMP4BoxBase>& InParent, const FMP4BoxInfo& InBoxInfo)
+			FMP4BoxMETA(const TWeakPtr<FMP4BoxBase, ESPMode::ThreadSafe>& InParent, const FMP4BoxInfo& InBoxInfo)
 				: FMP4BoxBase(InParent, InBoxInfo)
 			{ }
 		};
@@ -1437,12 +1437,12 @@ namespace Electra
 		class FMP4BoxSINF : public FMP4BoxBase
 		{
 		public:
-			static TSharedPtr<FMP4BoxBase> Create(const TWeakPtr<FMP4BoxBase>& InParent, const FMP4BoxInfo& InBoxInfo)
+			static TSharedPtr<FMP4BoxBase, ESPMode::ThreadSafe> Create(const TWeakPtr<FMP4BoxBase, ESPMode::ThreadSafe>& InParent, const FMP4BoxInfo& InBoxInfo)
 			{ return MakeShareable<FMP4BoxSINF>(new FMP4BoxSINF(InParent, InBoxInfo)); }
 			virtual ~FMP4BoxSINF() = default;
 			bool IsLeafBox() const override { return false; }
 		protected:
-			FMP4BoxSINF(const TWeakPtr<FMP4BoxBase>& InParent, const FMP4BoxInfo& InBoxInfo)
+			FMP4BoxSINF(const TWeakPtr<FMP4BoxBase, ESPMode::ThreadSafe>& InParent, const FMP4BoxInfo& InBoxInfo)
 				: FMP4BoxBase(InParent, InBoxInfo)
 			{ }
 		};
@@ -1455,11 +1455,11 @@ namespace Electra
 		class FMP4BoxFRMA : public FMP4BoxBase
 		{
 		public:
-			static TSharedPtr<FMP4BoxBase> Create(const TWeakPtr<FMP4BoxBase>& InParent, const FMP4BoxInfo& InBoxInfo)
+			static TSharedPtr<FMP4BoxBase, ESPMode::ThreadSafe> Create(const TWeakPtr<FMP4BoxBase, ESPMode::ThreadSafe>& InParent, const FMP4BoxInfo& InBoxInfo)
 			{ return MakeShareable<FMP4BoxFRMA>(new FMP4BoxFRMA(InParent, InBoxInfo)); }
 			virtual ~FMP4BoxFRMA() = default;
 		protected:
-			FMP4BoxFRMA(const TWeakPtr<FMP4BoxBase>& InParent, const FMP4BoxInfo& InBoxInfo)
+			FMP4BoxFRMA(const TWeakPtr<FMP4BoxBase, ESPMode::ThreadSafe>& InParent, const FMP4BoxInfo& InBoxInfo)
 				: FMP4BoxBase(InParent, InBoxInfo)
 			{ }
 		};
@@ -1472,11 +1472,11 @@ namespace Electra
 		class FMP4BoxSCHM : public FMP4BoxBase
 		{
 		public:
-			static TSharedPtr<FMP4BoxBase> Create(const TWeakPtr<FMP4BoxBase>& InParent, const FMP4BoxInfo& InBoxInfo)
+			static TSharedPtr<FMP4BoxBase, ESPMode::ThreadSafe> Create(const TWeakPtr<FMP4BoxBase, ESPMode::ThreadSafe>& InParent, const FMP4BoxInfo& InBoxInfo)
 			{ return MakeShareable<FMP4BoxSCHM>(new FMP4BoxSCHM(InParent, InBoxInfo)); }
 			virtual ~FMP4BoxSCHM() = default;
 		protected:
-			FMP4BoxSCHM(const TWeakPtr<FMP4BoxBase>& InParent, const FMP4BoxInfo& InBoxInfo)
+			FMP4BoxSCHM(const TWeakPtr<FMP4BoxBase, ESPMode::ThreadSafe>& InParent, const FMP4BoxInfo& InBoxInfo)
 				: FMP4BoxBase(InParent, InBoxInfo)
 			{ }
 		};
@@ -1489,12 +1489,12 @@ namespace Electra
 		class FMP4BoxSCHI : public FMP4BoxBase
 		{
 		public:
-			static TSharedPtr<FMP4BoxBase> Create(const TWeakPtr<FMP4BoxBase>& InParent, const FMP4BoxInfo& InBoxInfo)
+			static TSharedPtr<FMP4BoxBase, ESPMode::ThreadSafe> Create(const TWeakPtr<FMP4BoxBase, ESPMode::ThreadSafe>& InParent, const FMP4BoxInfo& InBoxInfo)
 			{ return MakeShareable<FMP4BoxSCHI>(new FMP4BoxSCHI(InParent, InBoxInfo)); }
 			virtual ~FMP4BoxSCHI() = default;
 			bool IsLeafBox() const override { return false; }
 		protected:
-			FMP4BoxSCHI(const TWeakPtr<FMP4BoxBase>& InParent, const FMP4BoxInfo& InBoxInfo)
+			FMP4BoxSCHI(const TWeakPtr<FMP4BoxBase, ESPMode::ThreadSafe>& InParent, const FMP4BoxInfo& InBoxInfo)
 				: FMP4BoxBase(InParent, InBoxInfo)
 			{ }
 		};
@@ -1507,14 +1507,14 @@ namespace Electra
 		class FMP4BoxPSSH : public FMP4BoxBase
 		{
 		public:
-			static TSharedPtr<FMP4BoxBase> Create(const TWeakPtr<FMP4BoxBase>& InParent, const FMP4BoxInfo& InBoxInfo)
+			static TSharedPtr<FMP4BoxBase, ESPMode::ThreadSafe> Create(const TWeakPtr<FMP4BoxBase, ESPMode::ThreadSafe>& InParent, const FMP4BoxInfo& InBoxInfo)
 			{ return MakeShareable<FMP4BoxPSSH>(new FMP4BoxPSSH(InParent, InBoxInfo)); }
 			virtual ~FMP4BoxPSSH() = default;
 			ELECTRABASE_API TConstArrayView<uint8> GetSystemID();
 			ELECTRABASE_API const TArray<TConstArrayView<uint8>>& GetKIDs();
 			ELECTRABASE_API TConstArrayView<uint8> GetData();
 		protected:
-			FMP4BoxPSSH(const TWeakPtr<FMP4BoxBase>& InParent, const FMP4BoxInfo& InBoxInfo)
+			FMP4BoxPSSH(const TWeakPtr<FMP4BoxBase, ESPMode::ThreadSafe>& InParent, const FMP4BoxInfo& InBoxInfo)
 				: FMP4BoxBase(InParent, InBoxInfo)
 			{ }
 			void ParseIfRequired();
@@ -1530,7 +1530,7 @@ namespace Electra
 		class FMP4BoxTENC : public FMP4BoxBase
 		{
 		public:
-			static TSharedPtr<FMP4BoxBase> Create(const TWeakPtr<FMP4BoxBase>& InParent, const FMP4BoxInfo& InBoxInfo)
+			static TSharedPtr<FMP4BoxBase, ESPMode::ThreadSafe> Create(const TWeakPtr<FMP4BoxBase, ESPMode::ThreadSafe>& InParent, const FMP4BoxInfo& InBoxInfo)
 			{ return MakeShareable<FMP4BoxTENC>(new FMP4BoxTENC(InParent, InBoxInfo)); }
 			virtual ~FMP4BoxTENC() = default;
 			ELECTRABASE_API bool HasDefaultCryptBlockValues();
@@ -1541,7 +1541,7 @@ namespace Electra
 			ELECTRABASE_API TConstArrayView<uint8> GetDefaultKID();
 			ELECTRABASE_API TConstArrayView<uint8> GetDefaultConstantIV();
 		protected:
-			FMP4BoxTENC(const TWeakPtr<FMP4BoxBase>& InParent, const FMP4BoxInfo& InBoxInfo)
+			FMP4BoxTENC(const TWeakPtr<FMP4BoxBase, ESPMode::ThreadSafe>& InParent, const FMP4BoxInfo& InBoxInfo)
 				: FMP4BoxBase(InParent, InBoxInfo)
 			{ }
 			void ParseIfRequired();
@@ -1557,11 +1557,11 @@ namespace Electra
 		class FMP4BoxSENC : public FMP4BoxBase
 		{
 		public:
-			static TSharedPtr<FMP4BoxBase> Create(const TWeakPtr<FMP4BoxBase>& InParent, const FMP4BoxInfo& InBoxInfo)
+			static TSharedPtr<FMP4BoxBase, ESPMode::ThreadSafe> Create(const TWeakPtr<FMP4BoxBase, ESPMode::ThreadSafe>& InParent, const FMP4BoxInfo& InBoxInfo)
 			{ return MakeShareable<FMP4BoxSENC>(new FMP4BoxSENC(InParent, InBoxInfo)); }
 			virtual ~FMP4BoxSENC() = default;
 
-			ELECTRABASE_API void Prepare(const TArray<TSharedPtr<Electra::UtilitiesMP4::FMP4BoxBase>>& InRelatedBoxes);
+			ELECTRABASE_API void Prepare(const TArray<TSharedPtr<Electra::UtilitiesMP4::FMP4BoxBase, ESPMode::ThreadSafe>>& InRelatedBoxes);
 
 			struct FSubSample
 			{
@@ -1575,7 +1575,7 @@ namespace Electra
 			};
 			ELECTRABASE_API const TArray<FEntry>& GetEntries();
 		protected:
-			FMP4BoxSENC(const TWeakPtr<FMP4BoxBase>& InParent, const FMP4BoxInfo& InBoxInfo)
+			FMP4BoxSENC(const TWeakPtr<FMP4BoxBase, ESPMode::ThreadSafe>& InParent, const FMP4BoxInfo& InBoxInfo)
 				: FMP4BoxBase(InParent, InBoxInfo)
 			{ }
 			void ParseIfRequired();
@@ -1593,12 +1593,12 @@ namespace Electra
 		class FMP4BoxAVCC : public FMP4BoxBase
 		{
 		public:
-			static TSharedPtr<FMP4BoxBase> Create(const TWeakPtr<FMP4BoxBase>& InParent, const FMP4BoxInfo& InBoxInfo)
+			static TSharedPtr<FMP4BoxBase, ESPMode::ThreadSafe> Create(const TWeakPtr<FMP4BoxBase, ESPMode::ThreadSafe>& InParent, const FMP4BoxInfo& InBoxInfo)
 			{ return MakeShareable<FMP4BoxAVCC>(new FMP4BoxAVCC(InParent, InBoxInfo)); }
 			virtual ~FMP4BoxAVCC() = default;
 			ELECTRABASE_API TConstArrayView<uint8> GetAVCDecoderConfigurationRecord();
 		protected:
-			FMP4BoxAVCC(const TWeakPtr<FMP4BoxBase>& InParent, const FMP4BoxInfo& InBoxInfo)
+			FMP4BoxAVCC(const TWeakPtr<FMP4BoxBase, ESPMode::ThreadSafe>& InParent, const FMP4BoxInfo& InBoxInfo)
 				: FMP4BoxBase(InParent, InBoxInfo)
 			{ }
 			void ParseIfRequired();
@@ -1614,12 +1614,12 @@ namespace Electra
 		class FMP4BoxHVCC : public FMP4BoxBase
 		{
 		public:
-			static TSharedPtr<FMP4BoxBase> Create(const TWeakPtr<FMP4BoxBase>& InParent, const FMP4BoxInfo& InBoxInfo)
+			static TSharedPtr<FMP4BoxBase, ESPMode::ThreadSafe> Create(const TWeakPtr<FMP4BoxBase, ESPMode::ThreadSafe>& InParent, const FMP4BoxInfo& InBoxInfo)
 			{ return MakeShareable<FMP4BoxHVCC>(new FMP4BoxHVCC(InParent, InBoxInfo)); }
 			virtual ~FMP4BoxHVCC() = default;
 			ELECTRABASE_API TConstArrayView<uint8> GetHEVCDecoderConfigurationRecord();
 		protected:
-			FMP4BoxHVCC(const TWeakPtr<FMP4BoxBase>& InParent, const FMP4BoxInfo& InBoxInfo)
+			FMP4BoxHVCC(const TWeakPtr<FMP4BoxBase, ESPMode::ThreadSafe>& InParent, const FMP4BoxInfo& InBoxInfo)
 				: FMP4BoxBase(InParent, InBoxInfo)
 			{ }
 			void ParseIfRequired();
@@ -1635,12 +1635,12 @@ namespace Electra
 		class FMP4BoxIODS : public FMP4BoxBase
 		{
 		public:
-			static TSharedPtr<FMP4BoxBase> Create(const TWeakPtr<FMP4BoxBase>& InParent, const FMP4BoxInfo& InBoxInfo)
+			static TSharedPtr<FMP4BoxBase, ESPMode::ThreadSafe> Create(const TWeakPtr<FMP4BoxBase, ESPMode::ThreadSafe>& InParent, const FMP4BoxInfo& InBoxInfo)
 			{ return MakeShareable<FMP4BoxIODS>(new FMP4BoxIODS(InParent, InBoxInfo)); }
 			virtual ~FMP4BoxIODS() = default;
 			ELECTRABASE_API TConstArrayView<uint8> GetObjectDescriptor();
 		protected:
-			FMP4BoxIODS(const TWeakPtr<FMP4BoxBase>& InParent, const FMP4BoxInfo& InBoxInfo)
+			FMP4BoxIODS(const TWeakPtr<FMP4BoxBase, ESPMode::ThreadSafe>& InParent, const FMP4BoxInfo& InBoxInfo)
 				: FMP4BoxBase(InParent, InBoxInfo)
 			{ }
 			void ParseIfRequired();
@@ -1656,12 +1656,12 @@ namespace Electra
 		class FMP4BoxESDS : public FMP4BoxBase
 		{
 		public:
-			static TSharedPtr<FMP4BoxBase> Create(const TWeakPtr<FMP4BoxBase>& InParent, const FMP4BoxInfo& InBoxInfo)
+			static TSharedPtr<FMP4BoxBase, ESPMode::ThreadSafe> Create(const TWeakPtr<FMP4BoxBase, ESPMode::ThreadSafe>& InParent, const FMP4BoxInfo& InBoxInfo)
 			{ return MakeShareable<FMP4BoxESDS>(new FMP4BoxESDS(InParent, InBoxInfo)); }
 			virtual ~FMP4BoxESDS() = default;
 			ELECTRABASE_API TConstArrayView<uint8> GetESDescriptor();
 		protected:
-			FMP4BoxESDS(const TWeakPtr<FMP4BoxBase>& InParent, const FMP4BoxInfo& InBoxInfo)
+			FMP4BoxESDS(const TWeakPtr<FMP4BoxBase, ESPMode::ThreadSafe>& InParent, const FMP4BoxInfo& InBoxInfo)
 				: FMP4BoxBase(InParent, InBoxInfo)
 			{ }
 			void ParseIfRequired();
@@ -1677,12 +1677,12 @@ namespace Electra
 		class FMP4BoxDVCC : public FMP4BoxBase
 		{
 		public:
-			static TSharedPtr<FMP4BoxBase> Create(const TWeakPtr<FMP4BoxBase>& InParent, const FMP4BoxInfo& InBoxInfo)
+			static TSharedPtr<FMP4BoxBase, ESPMode::ThreadSafe> Create(const TWeakPtr<FMP4BoxBase, ESPMode::ThreadSafe>& InParent, const FMP4BoxInfo& InBoxInfo)
 			{ return MakeShareable<FMP4BoxDVCC>(new FMP4BoxDVCC(InParent, InBoxInfo)); }
 			virtual ~FMP4BoxDVCC() = default;
 			ELECTRABASE_API TConstArrayView<uint8> GetDOVIDecoderConfigurationRecord();
 		protected:
-			FMP4BoxDVCC(const TWeakPtr<FMP4BoxBase>& InParent, const FMP4BoxInfo& InBoxInfo)
+			FMP4BoxDVCC(const TWeakPtr<FMP4BoxBase, ESPMode::ThreadSafe>& InParent, const FMP4BoxInfo& InBoxInfo)
 				: FMP4BoxBase(InParent, InBoxInfo)
 			{ }
 			void ParseIfRequired();
@@ -1698,12 +1698,12 @@ namespace Electra
 		class FMP4BoxDAC3 : public FMP4BoxBase
 		{
 		public:
-			static TSharedPtr<FMP4BoxBase> Create(const TWeakPtr<FMP4BoxBase>& InParent, const FMP4BoxInfo& InBoxInfo)
+			static TSharedPtr<FMP4BoxBase, ESPMode::ThreadSafe> Create(const TWeakPtr<FMP4BoxBase, ESPMode::ThreadSafe>& InParent, const FMP4BoxInfo& InBoxInfo)
 			{ return MakeShareable<FMP4BoxDAC3>(new FMP4BoxDAC3(InParent, InBoxInfo)); }
 			virtual ~FMP4BoxDAC3() = default;
 			ELECTRABASE_API TConstArrayView<uint8> GetAC3SpecificBox();
 		protected:
-			FMP4BoxDAC3(const TWeakPtr<FMP4BoxBase>& InParent, const FMP4BoxInfo& InBoxInfo)
+			FMP4BoxDAC3(const TWeakPtr<FMP4BoxBase, ESPMode::ThreadSafe>& InParent, const FMP4BoxInfo& InBoxInfo)
 				: FMP4BoxBase(InParent, InBoxInfo)
 			{ }
 			void ParseIfRequired();
@@ -1719,12 +1719,12 @@ namespace Electra
 		class FMP4BoxDEC3 : public FMP4BoxBase
 		{
 		public:
-			static TSharedPtr<FMP4BoxBase> Create(const TWeakPtr<FMP4BoxBase>& InParent, const FMP4BoxInfo& InBoxInfo)
+			static TSharedPtr<FMP4BoxBase, ESPMode::ThreadSafe> Create(const TWeakPtr<FMP4BoxBase, ESPMode::ThreadSafe>& InParent, const FMP4BoxInfo& InBoxInfo)
 			{ return MakeShareable<FMP4BoxDEC3>(new FMP4BoxDEC3(InParent, InBoxInfo)); }
 			virtual ~FMP4BoxDEC3() = default;
 			ELECTRABASE_API TConstArrayView<uint8> GetEC3SpecificBox();
 		protected:
-			FMP4BoxDEC3(const TWeakPtr<FMP4BoxBase>& InParent, const FMP4BoxInfo& InBoxInfo)
+			FMP4BoxDEC3(const TWeakPtr<FMP4BoxBase, ESPMode::ThreadSafe>& InParent, const FMP4BoxInfo& InBoxInfo)
 				: FMP4BoxBase(InParent, InBoxInfo)
 			{ }
 			void ParseIfRequired();
@@ -1740,12 +1740,12 @@ namespace Electra
 		class FMP4BoxDFLA : public FMP4BoxBase
 		{
 		public:
-			static TSharedPtr<FMP4BoxBase> Create(const TWeakPtr<FMP4BoxBase>& InParent, const FMP4BoxInfo& InBoxInfo)
+			static TSharedPtr<FMP4BoxBase, ESPMode::ThreadSafe> Create(const TWeakPtr<FMP4BoxBase, ESPMode::ThreadSafe>& InParent, const FMP4BoxInfo& InBoxInfo)
 			{ return MakeShareable<FMP4BoxDFLA>(new FMP4BoxDFLA(InParent, InBoxInfo)); }
 			virtual ~FMP4BoxDFLA() = default;
 			ELECTRABASE_API TConstArrayView<uint8> GetFLACSpecificBox();
 		protected:
-			FMP4BoxDFLA(const TWeakPtr<FMP4BoxBase>& InParent, const FMP4BoxInfo& InBoxInfo)
+			FMP4BoxDFLA(const TWeakPtr<FMP4BoxBase, ESPMode::ThreadSafe>& InParent, const FMP4BoxInfo& InBoxInfo)
 				: FMP4BoxBase(InParent, InBoxInfo)
 			{ }
 			void ParseIfRequired();
@@ -1761,12 +1761,12 @@ namespace Electra
 		class FMP4BoxDOPS : public FMP4BoxBase
 		{
 		public:
-			static TSharedPtr<FMP4BoxBase> Create(const TWeakPtr<FMP4BoxBase>& InParent, const FMP4BoxInfo& InBoxInfo)
+			static TSharedPtr<FMP4BoxBase, ESPMode::ThreadSafe> Create(const TWeakPtr<FMP4BoxBase, ESPMode::ThreadSafe>& InParent, const FMP4BoxInfo& InBoxInfo)
 			{ return MakeShareable<FMP4BoxDOPS>(new FMP4BoxDOPS(InParent, InBoxInfo)); }
 			virtual ~FMP4BoxDOPS() = default;
 			ELECTRABASE_API TConstArrayView<uint8> GetOpusSpecificBox();
 		protected:
-			FMP4BoxDOPS(const TWeakPtr<FMP4BoxBase>& InParent, const FMP4BoxInfo& InBoxInfo)
+			FMP4BoxDOPS(const TWeakPtr<FMP4BoxBase, ESPMode::ThreadSafe>& InParent, const FMP4BoxInfo& InBoxInfo)
 				: FMP4BoxBase(InParent, InBoxInfo)
 			{ }
 			void ParseIfRequired();
@@ -1782,12 +1782,12 @@ namespace Electra
 		class FMP4BoxVPCC : public FMP4BoxBase
 		{
 		public:
-			static TSharedPtr<FMP4BoxBase> Create(const TWeakPtr<FMP4BoxBase>& InParent, const FMP4BoxInfo& InBoxInfo)
+			static TSharedPtr<FMP4BoxBase, ESPMode::ThreadSafe> Create(const TWeakPtr<FMP4BoxBase, ESPMode::ThreadSafe>& InParent, const FMP4BoxInfo& InBoxInfo)
 			{ return MakeShareable<FMP4BoxVPCC>(new FMP4BoxVPCC(InParent, InBoxInfo)); }
 			virtual ~FMP4BoxVPCC() = default;
 			ELECTRABASE_API TConstArrayView<uint8> GetVPCodecConfigurationBox();
 		protected:
-			FMP4BoxVPCC(const TWeakPtr<FMP4BoxBase>& InParent, const FMP4BoxInfo& InBoxInfo)
+			FMP4BoxVPCC(const TWeakPtr<FMP4BoxBase, ESPMode::ThreadSafe>& InParent, const FMP4BoxInfo& InBoxInfo)
 				: FMP4BoxBase(InParent, InBoxInfo)
 			{ }
 			void ParseIfRequired();
@@ -1803,12 +1803,12 @@ namespace Electra
 		class FMP4BoxWAVE : public FMP4BoxBase
 		{
 		public:
-			static TSharedPtr<FMP4BoxBase> Create(const TWeakPtr<FMP4BoxBase>& InParent, const FMP4BoxInfo& InBoxInfo)
+			static TSharedPtr<FMP4BoxBase, ESPMode::ThreadSafe> Create(const TWeakPtr<FMP4BoxBase, ESPMode::ThreadSafe>& InParent, const FMP4BoxInfo& InBoxInfo)
 			{ return MakeShareable<FMP4BoxWAVE>(new FMP4BoxWAVE(InParent, InBoxInfo)); }
 			virtual ~FMP4BoxWAVE() = default;
 			bool IsLeafBox() const override { return false; }
 		protected:
-			FMP4BoxWAVE(const TWeakPtr<FMP4BoxBase>& InParent, const FMP4BoxInfo& InBoxInfo)
+			FMP4BoxWAVE(const TWeakPtr<FMP4BoxBase, ESPMode::ThreadSafe>& InParent, const FMP4BoxInfo& InBoxInfo)
 				: FMP4BoxBase(InParent, InBoxInfo)
 			{ }
 		};
@@ -1821,12 +1821,12 @@ namespace Electra
 		class FMP4BoxTAPT : public FMP4BoxBase
 		{
 		public:
-			static TSharedPtr<FMP4BoxBase> Create(const TWeakPtr<FMP4BoxBase>& InParent, const FMP4BoxInfo& InBoxInfo)
+			static TSharedPtr<FMP4BoxBase, ESPMode::ThreadSafe> Create(const TWeakPtr<FMP4BoxBase, ESPMode::ThreadSafe>& InParent, const FMP4BoxInfo& InBoxInfo)
 			{ return MakeShareable<FMP4BoxTAPT>(new FMP4BoxTAPT(InParent, InBoxInfo)); }
 			virtual ~FMP4BoxTAPT() = default;
 			bool IsLeafBox() const override { return false; }
 		protected:
-			FMP4BoxTAPT(const TWeakPtr<FMP4BoxBase>& InParent, const FMP4BoxInfo& InBoxInfo)
+			FMP4BoxTAPT(const TWeakPtr<FMP4BoxBase, ESPMode::ThreadSafe>& InParent, const FMP4BoxInfo& InBoxInfo)
 				: FMP4BoxBase(InParent, InBoxInfo)
 			{ }
 		};
@@ -1836,10 +1836,10 @@ namespace Electra
 		{
 		public:
 			static ELECTRABASE_API FMP4BoxFactory& Get();
-			ELECTRABASE_API TSharedPtr<FMP4BoxBase> Create(const TWeakPtr<FMP4BoxBase>& InParent, const FMP4BoxInfo& InBoxInfo);
+			ELECTRABASE_API TSharedPtr<FMP4BoxBase, ESPMode::ThreadSafe> Create(const TWeakPtr<FMP4BoxBase, ESPMode::ThreadSafe>& InParent, const FMP4BoxInfo& InBoxInfo);
 		private:
 			FMP4BoxFactory();
-			TMap<uint32, TSharedPtr<FMP4BoxBase>(*)(const TWeakPtr<FMP4BoxBase>& /*InParent*/, const FMP4BoxInfo& /*InBoxInfo*/)> FactoryMap;
+			TMap<uint32, TSharedPtr<FMP4BoxBase, ESPMode::ThreadSafe>(*)(const TWeakPtr<FMP4BoxBase, ESPMode::ThreadSafe>& /*InParent*/, const FMP4BoxInfo& /*InBoxInfo*/)> FactoryMap;
 		};
 
 	} // namespace UtilitiesMP4
diff --git a/Engine/Plugins/Media/ElectraUtil/Source/ElectraBase/Public/Utilities/MP4Boxes/MP4Track.h b/Engine/Plugins/Media/ElectraUtil/Source/ElectraBase/Public/Utilities/MP4Boxes/MP4Track.h
--- a/Engine/Plugins/Media/ElectraUtil/Source/ElectraBase/Public/Utilities/MP4Boxes/MP4Track.h
+++ b/Engine/Plugins/Media/ElectraUtil/Source/ElectraBase/Public/Utilities/MP4Boxes/MP4Track.h
@@ -47,7 +47,7 @@ namespace Electra
 			};
 
 		public:
-			static TSharedPtr<FMP4Track> Create(const TSharedPtr<FMP4BoxTRAK>& InTrakBox)
+			static TSharedPtr<FMP4Track, ESPMode::ThreadSafe> Create(const TSharedPtr<FMP4BoxTRAK, ESPMode::ThreadSafe>& InTrakBox)
 			{ return MakeShareable<FMP4Track>(new FMP4Track(InTrakBox)); }
 			virtual ~FMP4Track() = default;
 
@@ -62,7 +62,7 @@ namespace Electra
 
 				bool IsValid() const
 				{ return Convs.NumTotalSamples != 0; }
-				TWeakPtr<FMP4Track> GetTrack() const
+				TWeakPtr<FMP4Track, ESPMode::ThreadSafe> GetTrack() const
 				{ return Track; }
 				uint32 GetTrackID() const
 				{ return Convs.TrackID; }
@@ -134,7 +134,7 @@ namespace Electra
 				{ return SampleNumber+1 >= Convs.LastSample.LastSampleNumber; }
 
 				// Creates a copy of this iterator.
-				TSharedPtr<FIterator> Clone()
+				TSharedPtr<FIterator, ESPMode::ThreadSafe> Clone()
 				{ return MakeShareable(new FIterator(*this)); }
 			protected:
 				FIterator() = default;
@@ -143,7 +143,7 @@ namespace Electra
 
 				void Update();
 
-				TWeakPtr<FMP4Track> Track;
+				TWeakPtr<FMP4Track, ESPMode::ThreadSafe> Track;
 				FConvenience Convs;
 				FSTSZBoxIterator stszIt;
 				FSTTSBoxIterator sttsIt;
@@ -166,14 +166,14 @@ namespace Electra
 			};
 
 			// Create an interator starting at the first sample.
-			ELECTRABASE_API TSharedPtr<FIterator> CreateIterator();
+			ELECTRABASE_API TSharedPtr<FIterator, ESPMode::ThreadSafe> CreateIterator();
 			// Create an interator starting at the last sample (used when iterating in reverse, crossing back from the beginning to the end)
-			ELECTRABASE_API TSharedPtr<FIterator> CreateIteratorAtLastFrame();
+			ELECTRABASE_API TSharedPtr<FIterator, ESPMode::ThreadSafe> CreateIteratorAtLastFrame();
 			// Create an iterator starting at a keyframe on or before the given time, or at a later time within the given
 			// threshold should one be right after the given time and would not be selected due to timescale rounding issues.
-			ELECTRABASE_API TSharedPtr<FIterator> CreateIteratorAtKeyframe(FTimeValue InForTime, FTimeValue InLaterTimeThreshold=FTimeValue::GetZero());
+			ELECTRABASE_API TSharedPtr<FIterator, ESPMode::ThreadSafe> CreateIteratorAtKeyframe(FTimeValue InForTime, FTimeValue InLaterTimeThreshold=FTimeValue::GetZero());
 			// Create an interator starting at a given sample number.
-			ELECTRABASE_API TSharedPtr<FIterator> CreateIterator(uint32 InAtSampleNumber);
+			ELECTRABASE_API TSharedPtr<FIterator, ESPMode::ThreadSafe> CreateIterator(uint32 InAtSampleNumber);
 
 			// Returns the number of samples in this track.
 			ELECTRABASE_API uint32 GetNumberOfSamples();
@@ -186,24 +186,24 @@ namespace Electra
 			{ return Convs.LastSample; }
 
 		protected:
-			FMP4Track(const TSharedPtr<FMP4BoxTRAK>& InTrakBox) : TrakBox(InTrakBox)
+			FMP4Track(const TSharedPtr<FMP4BoxTRAK, ESPMode::ThreadSafe>& InTrakBox) : TrakBox(InTrakBox)
 			{ }
 
-			TSharedPtr<FMP4BoxTRAK> TrakBox;
-
-			TSharedPtr<FMP4BoxTKHD> TkhdBox;
-			TSharedPtr<FMP4BoxELST> ElstBox;
-			TSharedPtr<FMP4BoxMDHD> MdhdBox;
-			TSharedPtr<FMP4BoxSTTS> SttsBox;
-			TSharedPtr<FMP4BoxCTTS> CttsBox;
-			TSharedPtr<FMP4BoxSTSC> StscBox;
-			TSharedPtr<FMP4BoxSTSZ> StszBox;
-			TSharedPtr<FMP4BoxSTCO> StcoBox;
-			TSharedPtr<FMP4BoxSTSS> StssBox;
-			TSharedPtr<FMP4BoxUDTA> UdtaBox;
-
-			TArray<TSharedPtr<FMP4BoxSGPD>> SgpdBoxes;
-			TArray<TSharedPtr<FMP4BoxSBGP>> SbgpBoxes;
+			TSharedPtr<FMP4BoxTRAK, ESPMode::ThreadSafe> TrakBox;
+
+			TSharedPtr<FMP4BoxTKHD, ESPMode::ThreadSafe> TkhdBox;
+			TSharedPtr<FMP4BoxELST, ESPMode::ThreadSafe> ElstBox;
+			TSharedPtr<FMP4BoxMDHD, ESPMode::ThreadSafe> MdhdBox;
+			TSharedPtr<FMP4BoxSTTS, ESPMode::ThreadSafe> SttsBox;
+			TSharedPtr<FMP4BoxCTTS, ESPMode::ThreadSafe> CttsBox;
+			TSharedPtr<FMP4BoxSTSC, ESPMode::ThreadSafe> StscBox;
+			TSharedPtr<FMP4BoxSTSZ, ESPMode::ThreadSafe> StszBox;
+			TSharedPtr<FMP4BoxSTCO, ESPMode::ThreadSafe> StcoBox;
+			TSharedPtr<FMP4BoxSTSS, ESPMode::ThreadSafe> StssBox;
+			TSharedPtr<FMP4BoxUDTA, ESPMode::ThreadSafe> UdtaBox;
+
+			TArray<TSharedPtr<FMP4BoxSGPD, ESPMode::ThreadSafe>> SgpdBoxes;
+			TArray<TSharedPtr<FMP4BoxSBGP, ESPMode::ThreadSafe>> SbgpBoxes;
 
 			FConvenience Convs;
 
