Improvements to new Android videodecoder path
- Repaired throttling to avoid Image buffer underflows
- Simplified image queue use
- Made sure that Images are processed before releasing them
- Closing HardwareBuffers
#rnx

diff --git a/Engine/Plugins/Media/ElectraUtil/Source/ElectraSamples/Private/ElectraTextureSampleAndroid.cpp b/Engine/Plugins/Media/ElectraUtil/Source/ElectraSamples/Private/ElectraTextureSampleAndroid.cpp
--- a/Engine/Plugins/Media/ElectraUtil/Source/ElectraSamples/Private/ElectraTextureSampleAndroid.cpp
+++ b/Engine/Plugins/Media/ElectraUtil/Source/ElectraSamples/Private/ElectraTextureSampleAndroid.cpp
@@ -89,38 +89,6 @@ public:
 	FElectraTextureSampleSupport();
 	~FElectraTextureSampleSupport();
 
-	struct FFrameUpdateInfo
-	{
-		~FFrameUpdateInfo()
-		{
-			if (ImageResources)
-			{
-				JNIEnv* JEnv = FAndroidApplication::GetJavaEnv();
-				//TODO: CALL "Release" - OR REMOVE THIS VERSION IF WE NEVER USE IT!
-				JEnv->DeleteGlobalRef(ImageResources);
-			}
-		}
-
-		jobject MoveImageResources()
-		{
-			jobject Ret = ImageResources;
-			ImageResources = nullptr;
-			return Ret;
-		}
-
-		int64   Timestamp = 0;
-		int64   Duration = 0;
-		float   UScale = 0.0f;
-		float   UOffset = 0.0f;
-		float   VScale = 0.0f;
-		float   VOffset = 0.0f;
-		bool    bFrameReady = false;
-		bool    bRegionChanged = false;
-		int     NumPending = 0;
-		jobject ImageResources = nullptr;
-	};
-
-	int32 GetFrameDataAndUpdateInfo(FFrameUpdateInfo& OutFrameUpdateInfo, FElectraTextureSample* InTargetSample);
 	int32 GetFrameData(FElectraTextureSample* InTargetSample);
 	jobject GetCodecSurface();
 
@@ -407,59 +375,6 @@ FElectraTextureSampleSupport::~FElectraTextureSampleSupport()
 	}
 }
 
-//-----------------------------------------------------------------------------
-/**
- *
- * @note Call this from an RHI thread! It will need a valid rendering environment!
- */
-int32 FElectraTextureSampleSupport::GetFrameDataAndUpdateInfo(FFrameUpdateInfo& OutFrameUpdateInfo, FElectraTextureSample* InTargetSample)
-{
-	// In case this is called with a ES renderer, we need to pass in the destination texture we'd like to be used to receive the data
-	// (for Vulkan we'll just receive a simple byte buffer)
-	int32 DestTexture = 0;
-	if (FRHITexture* Texture = InTargetSample->GetTexture())
-	{
-		DestTexture = *reinterpret_cast<int32*>(Texture->GetNativeResource());
-	}
-
-	// Update frame info and get data...
-	JNIEnv* JEnv = FAndroidApplication::GetJavaEnv();
-	jobject OutputInfo = JEnv->CallObjectMethod(Object, GetVideoFrameUpdateInfoFN.Method, DestTexture, InTargetSample->GetDim().X, InTargetSample->GetDim().Y, InTargetSample->GetFormat() == EMediaTextureSampleFormat::CharBGR10A2);
-	if (JEnv->ExceptionCheck())
-	{
-		JEnv->ExceptionDescribe();
-		JEnv->ExceptionClear();
-	}
-	// Failure will return no object.
-	if (OutputInfo != nullptr)
-	{
-		OutFrameUpdateInfo.Timestamp = JEnv->GetLongField(OutputInfo, FFrameUpdateInfo_Timestamp);
-		OutFrameUpdateInfo.Duration = JEnv->GetLongField(OutputInfo, FFrameUpdateInfo_Duration);
-		OutFrameUpdateInfo.UScale = JEnv->GetFloatField(OutputInfo, FFrameUpdateInfo_UScale);
-		OutFrameUpdateInfo.UOffset = JEnv->GetFloatField(OutputInfo, FFrameUpdateInfo_UOffset);
-		OutFrameUpdateInfo.VScale = JEnv->GetFloatField(OutputInfo, FFrameUpdateInfo_VScale);
-		OutFrameUpdateInfo.VOffset = JEnv->GetFloatField(OutputInfo, FFrameUpdateInfo_VOffset);
-		OutFrameUpdateInfo.bFrameReady = JEnv->GetBooleanField(OutputInfo, FFrameUpdateInfo_bFrameReady);
-		OutFrameUpdateInfo.bRegionChanged = JEnv->GetBooleanField(OutputInfo, FFrameUpdateInfo_bRegionChanged);
-		OutFrameUpdateInfo.NumPending = JEnv->GetIntField(OutputInfo, FFrameUpdateInfo_NumPending);
-		if (InTargetSample)
-		{
-			jobject buffer = JEnv->GetObjectField(OutputInfo, FFrameUpdateInfo_Buffer);
-			if (buffer != nullptr)
-			{
-				const void* outPixels = JEnv->GetDirectBufferAddress(buffer);
-				const int32 outCount = JEnv->GetDirectBufferCapacity(buffer);
-				InTargetSample->SetupFromBuffer(outPixels, outCount);
-				JEnv->DeleteLocalRef(buffer);
-			}
-		}
-
-		JEnv->DeleteLocalRef(OutputInfo);
-		return 0;
-	}
-	return 1;
-}
-
 //-----------------------------------------------------------------------------
 /**
  *
@@ -606,6 +521,7 @@ void FElectraTextureSample::Initialize(FVideoDecoderOutput* InVideoDecoderOutput
 
 	if (Support->UseGpuOutputPath())
 	{
+		bQueuedForConversion = false;
 		Support->GetFrameData(this);
 		Texture = nullptr;
 	}
@@ -776,6 +692,23 @@ void FElectraTextureSample::SetImageResources(jobject InImageResources)
 }
 
 #if !UE_SERVER
+bool FElectraTextureSample::IsReadyForReuse()
+{
+	if (Support->UseGpuOutputPath())
+	{
+		if (bQueuedForConversion)
+		{
+			if (Fence->Poll())
+			{
+				bQueuedForConversion = false;
+			}
+		}
+
+		return !bQueuedForConversion;
+	}
+	return true;
+}
+
 void FElectraTextureSample::ShutdownPoolable()
 {
 	IElectraTextureSampleBase::ShutdownPoolable();
@@ -784,6 +717,11 @@ void FElectraTextureSample::ShutdownPoolable()
 // 	 - MIGHT that starve the decoder for buffers (we only have a max. number of "Images", too)
 //  >> set it to NULL and we'll recreate it... less ambigous maybe...
 //	Texture = nullptr;
+	if (Fence)
+	{
+		Fence->Clear();
+	}
+
 	CleanupImageResources();
 }
 #endif //!UE_SERVER
@@ -812,7 +750,7 @@ void FElectraTextureSample::CleanupImageResources()
 		jmethodID ReleaseFN = Support->GetImageResources_ReleaseFN();
 		jobject Resources = ImageResources;
 
-		if (IsInRHIThread())
+		if (Support->UseGpuOutputPath() || IsInRHIThread())
 		{
 			CleanupImageResourcesJNI(Resources, ReleaseFN);
 		}
@@ -867,7 +805,7 @@ uint32 FElectraTextureSample::GetStride() const
 }
 
 
-void  FElectraTextureSample::CopyFromExternalTextureOES(FRHICommandListImmediate& RHICmdList, FTextureRHIRef& InDstTexture, FTextureRHIRef& InSrcTexture, const FVector2f& InScale, const FVector2f& InOffset)
+void FElectraTextureSample::CopyFromExternalTextureOES(FRHICommandListImmediate& RHICmdList, FTextureRHIRef& InDstTexture, FTextureRHIRef& InSrcTexture, const FVector2f& InScale, const FVector2f& InOffset)
 {
 	FLinearColor Offset = { InOffset.X, InOffset.Y, 0.f, 0.f };
 	FLinearColor ScaleRotation = { InScale.X, 0.f,
@@ -924,6 +862,11 @@ void  FElectraTextureSample::CopyFromExternalTextureVulkan(FRHICommandListImmedi
 	FGraphicsPipelineStateInitializer GraphicsPSOInit;
 	FRHITexture* RenderTarget = InDstTexture;
 
+	if (!Fence)
+	{
+		Fence = RHICreateGPUFence(TEXT("CopyFromExternalTextureVulkan"));
+	}
+
 	RHICmdList.Transition(FRHITransitionInfo(InDstTexture, ERHIAccess::Unknown, ERHIAccess::RTV));
 
 	FRHIRenderPassInfo RPInfo(RenderTarget, ERenderTargetActions::DontLoad_Store);
@@ -987,13 +930,18 @@ void  FElectraTextureSample::CopyFromExternalTextureVulkan(FRHICommandListImmedi
 	RHICmdList.EndRenderPass();
 
 	RHICmdList.Transition(FRHITransitionInfo(RenderTarget, ERHIAccess::RTV, ERHIAccess::SRVGraphics));
+
+	RHICmdList.WriteGPUFence(Fence);
 }
 
 
 bool FElectraTextureSample::Convert(FRHICommandListImmediate& RHICmdList, FTextureRHIRef& InDstTexture, const FConversionHints& Hints)
 {
 	if (Support->UseGpuOutputPath())
-		return ConvertGpuOutputPath(RHICmdList, InDstTexture, Hints);
+	{
+		bQueuedForConversion = ConvertGpuOutputPath(RHICmdList, InDstTexture, Hints);
+		return bQueuedForConversion;
+	}
 	else
 		return ConvertCpuOutputPath(RHICmdList, InDstTexture, Hints);
 }
@@ -1011,6 +959,7 @@ bool FElectraTextureSample::ConvertGpuOutputPath(FRHICommandListImmediate& RHICm
 	jobject HardwareBufferObj = Support->ImageResources_GetHardwareBuffer(ImageResources);
 	AHardwareBuffer* HardwareBuffer = AHardwareBuffer_fromHardwareBuffer(FAndroidApplication::GetJavaEnv(), HardwareBufferObj);
 	ensure(HardwareBuffer);
+	AHardwareBuffer_acquire(HardwareBuffer);
 
 	check(VideoDecoderOutputAndroid->GetOutputType() == FVideoDecoderOutputAndroid::EOutputType::DirectToSurfaceAsQueue);
 	if (FAndroidMisc::ShouldUseVulkan())
@@ -1040,6 +989,8 @@ bool FElectraTextureSample::ConvertGpuOutputPath(FRHICommandListImmediate& RHICm
 		CopyFromExternalTextureOES(RHICmdList, InDstTexture, SrcTexture, Scale, Offset);
 	}
 
+	AHardwareBuffer_release(HardwareBuffer);
+	FAndroidApplication::GetJavaEnv()->DeleteLocalRef(HardwareBufferObj);
 	return true;
 }
 
diff --git a/Engine/Plugins/Media/ElectraUtil/Source/ElectraSamples/Private/Java/src/com/epicgames/unreal/ElectraTextureSample.java b/Engine/Plugins/Media/ElectraUtil/Source/ElectraSamples/Private/Java/src/com/epicgames/unreal/ElectraTextureSample.java
--- a/Engine/Plugins/Media/ElectraUtil/Source/ElectraSamples/Private/Java/src/com/epicgames/unreal/ElectraTextureSample.java
+++ b/Engine/Plugins/Media/ElectraUtil/Source/ElectraSamples/Private/Java/src/com/epicgames/unreal/ElectraTextureSample.java
@@ -24,6 +24,12 @@ public class ElectraTextureSample
 	{
 		public void Release()
 		{
+			if (HardwareBuffer != null)
+			{
+				HardwareBuffer.close();
+				HardwareBuffer = null;
+			}
+
 			if (Image != null)
 			{
 				Image.close();
@@ -981,9 +987,7 @@ public class ElectraTextureSample
 		private java.nio.Buffer mFrameData = null;
 		private int mFrameDataSize = 0;
 		private android.media.ImageReader mImageReader = null;
-		private android.os.HandlerThread mHandlerThread = null;
 		private android.os.Handler mImageReaderHandler = null;
-		private java.util.concurrent.BlockingQueue<Image> mImageQueue = null;
 		private int mTextureWidth = -1;
 		private int mTextureHeight = -1;
 		private android.view.Surface mSurface = null;
@@ -1002,38 +1006,13 @@ public class ElectraTextureSample
 
 			mSurface = mImageReader.getSurface();
 
-			mImageQueue = new java.util.concurrent.LinkedBlockingQueue<>(maxQueueLength / 2);
-			// We need to spawn our own thread here, because the main looper cannot be relied on to collect images before the ImageReader queue overflows
-			 mImageReaderHandler = new android.os.Handler(android.os.Looper.getMainLooper());
-			//mHandlerThread = new android.os.HandlerThread("ElectraTextureSampleHandlerThread");
-			//mHandlerThread.start();
-			//mImageReaderHandler = new android.os.Handler(mHandlerThread.getLooper());
+			mImageReaderHandler = new android.os.Handler(android.os.Looper.getMainLooper());
 			mImageReader.setOnImageAvailableListener(this, mImageReaderHandler);
 		}
 
 		public void onImageAvailable(ImageReader reader)
 		{
-			// Get image data from Surface / Queue...
-			android.media.Image Image = null;
-			try {
-				Image =  mImageReader.acquireNextImage();
-			} catch(Exception e) {
-				Log.d(TAG, String.format("ImageReader.acquireNextImage() threw: %s", e.toString()));
-			}
-
-			if (Image != null)
-			{
-				// Signal that we received the image
-				// (used on pre API 31 to make sure we only have one image in the queue; NoOp later on!)
-				nativeSignalSurfaceReadEvent(mParentHandle);
-
-				// Deliver new image into our own queue
-				try {
-					mImageQueue.put(Image);
-				} catch(InterruptedException e) {
-					// [...] ?
-				}
-			}
+			nativeSignalSurfaceReadEvent(mParentHandle);
 		}
 
 		public boolean isValid()
@@ -1087,21 +1066,28 @@ public class ElectraTextureSample
 				return null;
 			}
 
-			// Prepare to capture any resources we need to track for the output image
-			FImageResources ImageResources = new FImageResources();
-
-			// Grab the next image in the queue
-			try {
-				ImageResources.Image = mImageQueue.take();
-			} catch(InterruptedException e) {
-				// [...] ?
+			android.media.Image Image = null;
+			try
+			{
+				// Get image data from Surface / Queue...
+				Image = mImageReader.acquireNextImage();
+			} 
+			catch (Exception e) 
+			{
+				Log.d(TAG, String.format("getFrameTextureInfo: ImageReader.acquireNextImage() threw: %s", e.toString()));
+				return null;
 			}
 
-			if (ImageResources.Image == null)
+			if (Image == null)
 			{
+				Log.d(TAG, "getFrameTextureInfo requested, but no images queued...");
 				return null;
 			}
 
+			// Prepare to capture any resources we need to track for the output image
+			FImageResources ImageResources = new FImageResources();
+			ImageResources.Image = Image;
+
 			ImageResources.HardwareBuffer = ImageResources.Image.getHardwareBuffer();
 			if (ImageResources.HardwareBuffer == null)
 			{
@@ -1109,40 +1095,6 @@ public class ElectraTextureSample
 				return null;
 			}
 
-/*
-//ALL THIS COULD GO 100% NATIVE POSSIBLY - ONCE WE OWN THE Image WE SHOULD BE GOOD?
-			// Grab the EGLNativeClientBuffer from the Image using native code
-
-			if (mVulkanRenderer)
-			{
-				ImageResources.EglImageNativeHandle = 0;
-				ImageResources.HardwareBuffer = ImageResources.Image.getHardwareBuffer();
-				if (ImageResources.HardwareBuffer == null)
-				{
-					Log.d(TAG, "Could not get HardwareBuffer from Image!");
-					return null;
-				}
-			}
-			else
-			{
-				ImageResources.HardwareBuffer = null;
-				long NativeClientBuffer = nativeGetEGLNativeClientBuffer(ImageResources.Image.getHardwareBuffer());
-				if (NativeClientBuffer < 0)
-				{
-					Log.d(TAG, "Could not get native client buffer!");
-					return null;
-				}
-
-				ImageResources.EglImageNativeHandle = nativeCreateEGLImageKHR(EGL14.eglGetCurrentDisplay().getNativeHandle(), EGL14.eglGetCurrentContext().getNativeHandle(), NativeClientBuffer);
-				if (ImageResources.EglImageNativeHandle == 0)
-				{
-					Log.d(TAG, String.format("Could not create EGLimage from native client buffer! B=0x%x E=0x%x", NativeClientBuffer, EGL14.eglGetError()));
-					return null;
-				}
-
-				Log.d(TAG, String.format("NATIVE HANDLE: 0x%x", ImageResources.EglImageNativeHandle));
-			}
-*/
 			android.graphics.Rect CropRect = ImageResources.Image.getCropRect();
 			ImageResources.UScale = (float)(CropRect.right - CropRect.left) / mTextureWidth;
 			ImageResources.UOffset = (float)CropRect.left / mTextureWidth;
@@ -1180,11 +1132,6 @@ public class ElectraTextureSample
 				mImageReader.close();
 				mImageReader = null;
 			}
-			if (mHandlerThread != null)
-			{
-				mHandlerThread.stop();
-				mHandlerThread = null;
-			}
 		}
 	}
 
diff --git a/Engine/Plugins/Media/ElectraUtil/Source/ElectraSamples/Public/Android/ElectraTextureSample.h b/Engine/Plugins/Media/ElectraUtil/Source/ElectraSamples/Public/Android/ElectraTextureSample.h
--- a/Engine/Plugins/Media/ElectraUtil/Source/ElectraSamples/Public/Android/ElectraTextureSample.h
+++ b/Engine/Plugins/Media/ElectraUtil/Source/ElectraSamples/Public/Android/ElectraTextureSample.h
@@ -43,6 +43,7 @@ public:
 	virtual ~FElectraTextureSample();
 
 #if !UE_SERVER
+	virtual bool IsReadyForReuse() override;
 	virtual void ShutdownPoolable() override;
 #endif
 
@@ -104,6 +105,9 @@ private:
 	TRefCountPtr<FRHITexture> Texture;
 
 	jobject ImageResources;
+	bool bQueuedForConversion;
+
+	TRefCountPtr<FRHIGPUFence> Fence;
 };
 
 
