Keeping hardwarebuffer wrapping texture around until fence is signalled
#rnx

diff --git a/Engine/Plugins/Media/ElectraUtil/Source/ElectraSamples/Private/ElectraTextureSampleAndroid.cpp b/Engine/Plugins/Media/ElectraUtil/Source/ElectraSamples/Private/ElectraTextureSampleAndroid.cpp
--- a/Engine/Plugins/Media/ElectraUtil/Source/ElectraSamples/Private/ElectraTextureSampleAndroid.cpp
+++ b/Engine/Plugins/Media/ElectraUtil/Source/ElectraSamples/Private/ElectraTextureSampleAndroid.cpp
@@ -717,6 +717,12 @@ void FElectraTextureSample::ShutdownPoolable()
 // 	 - MIGHT that starve the decoder for buffers (we only have a max. number of "Images", too)
 //  >> set it to NULL and we'll recreate it... less ambigous maybe...
 //	Texture = nullptr;
+
+	if (Support->UseGpuOutputPath())
+	{
+		Texture = nullptr;
+	}
+
 	if (Fence)
 	{
 		Fence->Clear();
@@ -955,7 +961,6 @@ bool FElectraTextureSample::ConvertGpuOutputPath(FRHICommandListImmediate& RHICm
 		return false;
 	}
 
-	FTextureRHIRef SrcTexture;
 	jobject HardwareBufferObj = Support->ImageResources_GetHardwareBuffer(ImageResources);
 	AHardwareBuffer* HardwareBuffer = AHardwareBuffer_fromHardwareBuffer(FAndroidApplication::GetJavaEnv(), HardwareBufferObj);
 	ensure(HardwareBuffer);
@@ -964,20 +969,21 @@ bool FElectraTextureSample::ConvertGpuOutputPath(FRHICommandListImmediate& RHICm
 	check(VideoDecoderOutputAndroid->GetOutputType() == FVideoDecoderOutputAndroid::EOutputType::DirectToSurfaceAsQueue);
 	if (FAndroidMisc::ShouldUseVulkan())
 	{
-		SrcTexture = InitializeTextureVulkan(HardwareBuffer);
+		Texture = InitializeTextureVulkan(HardwareBuffer);
 	}
 	else
 	{
-		SrcTexture = InitializeTextureOES(HardwareBuffer);
+		Texture = InitializeTextureOES(HardwareBuffer);
 	}
 
-	ensure(SrcTexture);
+	ensure(Texture);
+
 	if (FAndroidMisc::ShouldUseVulkan())
 	{
 		FVector2f Scale, Offset;
 		Support->ImageResources_GetScaleOffset(ImageResources, Scale, Offset);
 
-		CopyFromExternalTextureVulkan(RHICmdList, InDstTexture, SrcTexture, Scale, Offset);
+		CopyFromExternalTextureVulkan(RHICmdList, InDstTexture, Texture, Scale, Offset);
 	}
 	else
 	{
@@ -986,7 +992,7 @@ bool FElectraTextureSample::ConvertGpuOutputPath(FRHICommandListImmediate& RHICm
 
 		//FOR NOW(?) THIS IS DONE HERE TO MAKE SURE WE HAVE EASY ACCESS TO THE SCALE/OFFSET/ROTATION VALUES FOR EACH SAMPLE
 		//(the code using a map & GUID lookup assumes ONE "current" value per player... which does entirely NOT work in reality (queue of frames))
-		CopyFromExternalTextureOES(RHICmdList, InDstTexture, SrcTexture, Scale, Offset);
+		CopyFromExternalTextureOES(RHICmdList, InDstTexture, Texture, Scale, Offset);
 	}
 
 	AHardwareBuffer_release(HardwareBuffer);
