Fixed incorrect shader resource view creation caused by 38861284
-Fixed SA warnings in virtual texturing

diff --git a/Engine/Plugins/Media/ElectraUtil/Source/ElectraSamples/Private/ElectraTextureSampleAVF.cpp b/Engine/Plugins/Media/ElectraUtil/Source/ElectraSamples/Private/ElectraTextureSampleAVF.cpp
--- a/Engine/Plugins/Media/ElectraUtil/Source/ElectraSamples/Private/ElectraTextureSampleAVF.cpp
+++ b/Engine/Plugins/Media/ElectraUtil/Source/ElectraSamples/Private/ElectraTextureSampleAVF.cpp
@@ -258,13 +258,13 @@ void FElectraMediaTexConvApple::ConvertTexture(FTextureRHIRef & InDstTexture, CV
 						FShaderResourceViewRHIRef Y_SRV = RHICmdList.CreateShaderResourceView(
 							YTex, 
 							FRHIViewDesc::CreateTextureSRV()
-								.SetDimensionFromTexture(GetTextureInternal(0))
+								.SetDimensionFromTexture(YTex)
 								.SetMipRange(0, 1)
 								.SetFormat(PF_G8));
 						FShaderResourceViewRHIRef UV_SRV = RHICmdList.CreateShaderResourceView(
 							UVTex, 
 							FRHIViewDesc::CreateTextureSRV()
-								.SetDimensionFromTexture(GetTextureInternal(0))
+								.SetDimensionFromTexture(UVTex)
 								.SetMipRange(0, 1)
 								.SetFormat(PF_R8G8));
 
@@ -285,13 +285,13 @@ void FElectraMediaTexConvApple::ConvertTexture(FTextureRHIRef & InDstTexture, CV
 						FShaderResourceViewRHIRef Y_SRV = RHICmdList.CreateShaderResourceView(
 							YTex, 
 							FRHIViewDesc::CreateTextureSRV()
-								.SetDimensionFromTexture(GetTextureInternal(0))
+								.SetDimensionFromTexture(YTex)
 								.SetMipRange(0, 1)
 								.SetFormat(PF_G16));
 						FShaderResourceViewRHIRef UV_SRV = RHICmdList.CreateShaderResourceView(
 							UVTex, 
 							FRHIViewDesc::CreateTextureSRV()
-								.SetDimensionFromTexture(GetTextureInternal(0))
+								.SetDimensionFromTexture(UVTex)
 								.SetMipRange(0, 1)
 								.SetFormat(PF_G16R16));
 						
diff --git a/Engine/Source/Runtime/Apple/MetalRHI/Private/MetalTexture.cpp b/Engine/Source/Runtime/Apple/MetalRHI/Private/MetalTexture.cpp
--- a/Engine/Source/Runtime/Apple/MetalRHI/Private/MetalTexture.cpp
+++ b/Engine/Source/Runtime/Apple/MetalRHI/Private/MetalTexture.cpp
@@ -1972,7 +1972,11 @@ FTextureReferenceRHIRef FMetalDynamicRHI::RHICreateTextureReference(FRHICommandL
 		// If the referenced texture is configured for bindless, make sure we also create an SRV to use for bindless.
 		if (ReferencedTexture && ReferencedTexture->GetDefaultBindlessHandle().IsValid())
 		{
-			FShaderResourceViewRHIRef BindlessView = RHICmdList.CreateShaderResourceView(ReferencedTexture, 0u);
+			FShaderResourceViewRHIRef BindlessView = RHICmdList.CreateShaderResourceView(
+				ReferencedTexture,
+				FRHIViewDesc::CreateTextureSRV()
+					.SetDimensionFromTexture(ReferencedTexture)
+					.SetMipRange(0, 1));
 			return new FMetalTextureReference(ReferencedTexture, ResourceCast(BindlessView.GetReference()));
 		}
 	}
diff --git a/Engine/Source/Runtime/Renderer/Private/VT/VirtualTexturePhysicalSpace.cpp b/Engine/Source/Runtime/Renderer/Private/VT/VirtualTexturePhysicalSpace.cpp
--- a/Engine/Source/Runtime/Renderer/Private/VT/VirtualTexturePhysicalSpace.cpp
+++ b/Engine/Source/Runtime/Renderer/Private/VT/VirtualTexturePhysicalSpace.cpp
@@ -175,23 +175,19 @@ void FVirtualTexturePhysicalSpace::InitRHI(FRHICommandListBase& RHICmdList)
 		FRHITexture* TextureRHI = PooledRenderTarget[Layer]->GetRHI();
 
 		// Create sRGB and non-sRGB shader resource views into the physical texture
-		FRHITextureSRVCreateInfo SRVCreateInfo;
-		SRVCreateInfo.Format = FormatSRV;
-		SRVCreateInfo.SRGBOverride = SRGBO_ForceDisable;
 		TextureSRV[Layer] = RHICmdList.CreateShaderResourceView(
 			TextureRHI,
 			FRHIViewDesc::CreateTextureSRV()
 				.SetDimensionFromTexture(TextureRHI)
-				.SetFormat(SRVCreateInfo.Format)
-				.SetDisableSRGB(SRVCreateInfo.SRGBOverride == SRGBO_ForceDisable));
+				.SetFormat(FormatSRV)
+				.SetDisableSRGB(true));
 
-		SRVCreateInfo.SRGBOverride = SRGBO_Default;
 		TextureSRV_SRGB[Layer] = RHICmdList.CreateShaderResourceView(
 			TextureRHI,
 			FRHIViewDesc::CreateTextureSRV()
 				.SetDimensionFromTexture(TextureRHI)
-				.SetFormat(SRVCreateInfo.Format)
-				.SetDisableSRGB(SRVCreateInfo.SRGBOverride == SRGBO_ForceDisable));
+				.SetFormat(FormatSRV)
+				.SetDisableSRGB(false));
 	}
 
 	INC_MEMORY_STAT_BY(STAT_TotalPhysicalMemory, GetSizeInBytes());
