Fixed bug in ElectraTextureSample where delete was happening from the incorrect thread.

diff --git a/Engine/Plugins/Media/ElectraUtil/Source/ElectraSamples/Private/ElectraTextureSampleAVF.cpp b/Engine/Plugins/Media/ElectraUtil/Source/ElectraSamples/Private/ElectraTextureSampleAVF.cpp
--- a/Engine/Plugins/Media/ElectraUtil/Source/ElectraSamples/Private/ElectraTextureSampleAVF.cpp
+++ b/Engine/Plugins/Media/ElectraUtil/Source/ElectraSamples/Private/ElectraTextureSampleAVF.cpp
@@ -22,6 +22,10 @@ THIRD_PARTY_INCLUDES_START
 #include "MetalInclude.h"
 THIRD_PARTY_INCLUDES_END
 
+#if WITH_ENGINE
+extern void SafeReleaseMetalObject(NS::Object* Object);
+#endif
+
 // ------------------------------------------------------------------------------------------------------
 // ------------------------------------------------------------------------------------------------------
 // ------------------------------------------------------------------------------------------------------
@@ -36,16 +40,10 @@ FElectraMediaTexConvApple::FElectraMediaTexConvApple()
 FElectraMediaTexConvApple::~FElectraMediaTexConvApple()
 {
 #if WITH_ENGINE
-	// TODO: Carl - Test This  
 	if (MetalTextureCache)
 	{
 		CVMetalTextureCacheRef TextureCacheCopy = MetalTextureCache;
-		
-		FRHICommandListImmediate& RHICmdList = FRHICommandListImmediate::Get();
-		RHICmdList.EnqueueLambda([&TextureCacheCopy](FRHICommandListBase& RHICmdList)
-		{
-			TextureCacheCopy = nullptr;
-		});
+		SafeReleaseMetalObject((__bridge NS::Object*)TextureCacheCopy);
 	}
 #endif
 }
