Electra: Added multiple CDN support and ContentSteering for DASH and HLS

diff --git a/Engine/Plugins/Media/ElectraUtil/Source/ElectraBase/Private/URLParser.cpp b/Engine/Plugins/Media/ElectraUtil/Source/ElectraBase/Private/URLParser.cpp
--- a/Engine/Plugins/Media/ElectraUtil/Source/ElectraBase/Private/URLParser.cpp
+++ b/Engine/Plugins/Media/ElectraUtil/Source/ElectraBase/Private/URLParser.cpp
@@ -11,10 +11,12 @@ namespace Electra
 		static const FString KeepCharsFragment(TEXT("!$&'()*+,;=;@/?"));
 
 		static const FString URN(TEXT("urn:"));
+		static const FString Data(TEXT("data:"));
 		static const FString Slash(TEXT("/"));
 		static const FString DotSlash(TEXT("./"));
 		static const FString SlashSlash(TEXT("//"));
 		static const FString SchemeFile(TEXT("file"));
+		static const FString SchemeData(TEXT("data"));
 		static const FString SchemeHttp(TEXT("http"));
 		static const FString SchemeHttps(TEXT("https"));
 		static const FString Colon(TEXT(":"));
@@ -46,6 +48,8 @@ namespace Electra
 		Path.Empty();
 		Query.Empty();
 		Fragment.Empty();
+		bIsFile = false;
+		bIsData = false;
 	}
 
 	void FURL_RFC3986::Swap(FURL_RFC3986& Other)
@@ -57,6 +61,8 @@ namespace Electra
 		_SwapStrings(Path, Other.Path);
 		_SwapStrings(Query, Other.Query);
 		_SwapStrings(Fragment, Other.Fragment);
+		bool b = bIsFile; bIsFile = Other.bIsFile; Other.bIsFile = b;
+		b = bIsData; bIsData = Other.bIsData; Other.bIsData = b;
 	}
 
 	bool FURL_RFC3986::Parse(const FString& InURL)
@@ -71,6 +77,16 @@ namespace Electra
 		{
 			return false;
 		}
+
+		// Short circuit data URLs
+		if (InURL.StartsWith(URLComponents::Data))
+		{
+			bIsData = true;
+			Scheme = URLComponents::SchemeData;
+			Path = InURL.Mid(5);
+			return true;
+		}
+
 		/*
 		RFC 3986, section 3: Syntax components
 
@@ -116,6 +132,8 @@ namespace Electra
 				}
 				// Assume we parsed the scheme.
 				Scheme = MoveTemp(Component);
+				bIsFile = Scheme.Equals(URLComponents::SchemeFile);
+
 				// Does an absolute path ('/') or an authority ('//') follow?
 				if (IsPathSeparator(*it))
 				{
@@ -126,7 +144,7 @@ namespace Electra
 						++it;
 
 						// Check for a third '/' in Windows filenames like "file:///c:/autoexec.bat"
-						if (Scheme.Equals(URLComponents::SchemeFile) && IsPathSeparator(*it) && it.GetRemainingLength() > 3)
+						if (bIsFile && IsPathSeparator(*it) && it.GetRemainingLength() > 3)
 						{
 							const TCHAR* Next = it.GetRemainder();
 							if (Next[2] == TCHAR(':') && Next[3] == TCHAR('/'))
@@ -162,6 +180,35 @@ namespace Electra
 		return true;
 	}
 
+
+	void FURL_RFC3986::SetScheme(const FString& InNewScheme)
+	{
+		Scheme = InNewScheme;
+		bIsFile = Scheme.Equals(URLComponents::SchemeFile);
+		bIsData = Scheme.Equals(URLComponents::SchemeData);
+	}
+	void FURL_RFC3986::SetHost(const FString& InNewHost)
+	{
+		Host = InNewHost;
+	}
+	void FURL_RFC3986::SetPort(const FString& InNewPort)
+	{
+		Port = InNewPort;
+	}
+	void FURL_RFC3986::SetPath(const FString& InNewPath)
+	{
+		Path = InNewPath;
+	}
+	void FURL_RFC3986::SetQuery(const FString& InNewQuery)
+	{
+		Query = InNewQuery;
+	}
+	void FURL_RFC3986::SetFragment(const FString& InNewFragment)
+	{
+		Fragment = InNewFragment;
+	}
+
+
 	FString FURL_RFC3986::GetScheme() const
 	{
 		return Scheme;
@@ -194,11 +241,14 @@ namespace Electra
 
 	FString FURL_RFC3986::Get(bool bIncludeQuery, bool bIncludeFragment)
 	{
+		if (bIsData)
+		{
+			return FString::Printf(TEXT("%s%s"), *URLComponents::Data, *Path);
+		}
 		FString URL;
 		URL.Reserve(Scheme.Len() + UserInfo.Len() + Host.Len() + Port.Len() + Path.Len() + Query.Len() + Fragment.Len() + 32);
 		if (IsAbsolute())
 		{
-			bool bIsFile = Scheme.Equals(URLComponents::SchemeFile);
 			FString Authority = GetAuthority();
 			URL = Scheme;
 			URL += URLComponents::Colon;
@@ -239,6 +289,10 @@ namespace Electra
 
 	FString FURL_RFC3986::GetPath(bool bIncludeQuery, bool bIncludeFragment)
 	{
+		if (bIsData)
+		{
+			return Path;
+		}
 		FString URL;
 		URL.Reserve(Path.Len() + Query.Len() + Fragment.Len() + 32);
 		if (IsAbsolute())
@@ -323,10 +377,15 @@ namespace Electra
 	{
 		GetQueryParams(OutQueryParams, Query, bPerformUrlDecoding, bSameNameReplacesValue);
 	}
-	
+
 	void FURL_RFC3986::SetQueryParams(const TArray<FQueryParam>& InQueryParams)
 	{
 		Query.Empty();
+		if (bIsData)
+		{
+			return;
+		}
+
 		int32 Len=0;
 		for(int32 i=0; i<InQueryParams.Num(); ++i)
 		{
@@ -344,50 +403,74 @@ namespace Electra
 	}
 
 	// Appends or prepends additional query parameters.
-	void FURL_RFC3986::AddQueryParameters(const FString& InQueryParameters, bool bAppend)
+    void FURL_RFC3986::AddOrUpdateQueryParams(const TArray<FQueryParam>& InQueryParams)
 	{
-		if (InQueryParameters.Len() > 0)
+		if (bIsData)
+		{
+			return;
+		}
+		TArray<FQueryParam> Current;
+		GetQueryParams(Current, false, true);
+		for(int32 i=0; i<InQueryParams.Num(); ++i)
 		{
-			FString qp(InQueryParameters);
-			if (qp[0] == TCHAR('?') || qp[0] == TCHAR('&'))
+			bool bUpdated = false;
+			for(int32 j=0; j<Current.Num(); ++j)
 			{
-				qp.RightChopInline(1, EAllowShrinking::No);
+				if (Current[j].Name.Equals(InQueryParams[i].Name))
+				{
+					Current[j].Value = InQueryParams[i].Value;
+					bUpdated = true;
+					break;
+				}
 			}
-			if (qp.Len() && qp[qp.Len()-1] == TCHAR('&'))
+			if (!bUpdated)
 			{
-				qp.LeftChopInline(1, EAllowShrinking::No);
+				Current.Emplace(InQueryParams[i]);
 			}
-			if (bAppend)
+		}
+		SetQueryParams(Current);
+	}
+
+	void FURL_RFC3986::AddOrUpdateQueryParams(const FString& InQueryParameters)
+	{
+		if (InQueryParameters.Len() > 0)
+		{
+			FString _Query(InQueryParameters);
+			if (_Query[0] == TCHAR('?') || _Query[0] == TCHAR('&'))
 			{
-				if (Query.Len())
-				{
-					Query.AppendChar(TCHAR('&'));
-				}
-				Query.Append(qp);
+				_Query.RightChopInline(1, EAllowShrinking::No);
 			}
-			else
+			if (_Query.Len() && _Query[_Query.Len()-1] == TCHAR('&'))
 			{
-				if (Query.Len())
-				{
-					Query = qp + TCHAR('&') + Query;
-				}
-				else
-				{
-					Query = qp;
-				}
+				_Query.LeftChopInline(1, EAllowShrinking::No);
 			}
+			TArray<FQueryParam> qp;
+			GetQueryParams(qp, _Query, false, true);
+			AddOrUpdateQueryParams(qp);
 		}
 	}
 
-
 	bool FURL_RFC3986::IsAbsolute() const
 	{
 		return Scheme.Len() > 0;
 	}
 
-	void FURL_RFC3986::GetPathComponents(TArray<FString>& OutPathComponents) const
+	bool FURL_RFC3986::GetPathComponents(TArray<FString>& OutPathComponents) const
 	{
+		if (bIsData)
+		{
+			int32 CommaPos = INDEX_NONE;
+			if (Path.FindChar(TCHAR(','), CommaPos))
+			{
+				FString MediaType = Path.Mid(0, CommaPos);
+				MediaType.ParseIntoArray(OutPathComponents, TEXT(";"), true);
+				OutPathComponents.Emplace(Path.Mid(CommaPos + 1));
+				return true;
+			}
+			return false;
+		}
 		GetPathComponents(OutPathComponents, GetPath());
+		return true;
 	}
 
 	FString FURL_RFC3986::GetLastPathComponent() const
@@ -406,7 +489,7 @@ namespace Electra
 			Authority += UserInfo;
 			Authority += URLComponents::At;
 		}
-		if (Scheme.Equals(URLComponents::SchemeFile))
+		if (bIsFile)
 		{
 			Authority += Host;
 		}
@@ -434,7 +517,6 @@ namespace Electra
 		return Authority;
 	}
 
-
 	bool FURL_RFC3986::HasSameOriginAs(const FURL_RFC3986& Other)
 	{
 		/* RFC 6454:
@@ -450,7 +532,6 @@ namespace Electra
 		return Scheme.Equals(Other.Scheme) && Host.Equals(Other.Host) && Port.Equals(Other.Port);
 	}
 
-
 	bool FURL_RFC3986::ParseAuthority(StringHelpers::FStringIterator& it)
 	{
 		UserInfo.Empty();
@@ -483,7 +564,7 @@ namespace Electra
 		}
 		Host.Empty();
 		Host.Reserve(it.GetRemainingLength());
-		if (Scheme.Equals(URLComponents::SchemeFile))
+		if (bIsFile)
 		{
 			// For file:// scheme we have to consider Windows drive letters with a colon, eg. file://D:/
 			// We must not stop parsing at the colon since it will not indicate a port for the file scheme anyway.
@@ -711,6 +792,12 @@ namespace Electra
 		return FString();
 	}
 
+
+	FString FURL_RFC3986::GetUrlEncodeSubDelimsChars()
+	{
+		return FString(TEXT("!$&'()*+,;="));
+	}
+
 	void FURL_RFC3986::GetPathComponents(TArray<FString>& OutPathComponents, const FString& InPath)
 	{
 		TArray<FString> Components;
@@ -739,13 +826,25 @@ namespace Electra
 		}
 		return *this;
 	}
-	
+
 	/**
 	 * Resolve URL as per RFC 3986 section 5.2
 	 */
 	void FURL_RFC3986::ResolveWith(const FURL_RFC3986& Other)
 	{
-		if (!Other.Scheme.IsEmpty())
+		// data urls cannot be resolved, but we handle them here for consistency.
+		if (Other.bIsData)
+		{
+			Empty();
+			Scheme = Other.Scheme;
+			Path = Other.Path;
+		}
+		else if (bIsData && Other.Scheme.IsEmpty())
+		{
+			// Nothing to do.
+			return;
+		}
+		else if (!Other.Scheme.IsEmpty())
 		{
 			Scheme   = Other.Scheme;
 			UserInfo = Other.UserInfo;
@@ -791,6 +890,8 @@ namespace Electra
 			}
 		}
 		Fragment = Other.Fragment;
+		bIsFile = Scheme.Equals(URLComponents::SchemeFile);
+		bIsData = Scheme.Equals(URLComponents::SchemeData);
 	}
 
 	void FURL_RFC3986::BuildPathFromSegments(const TArray<FString>& Components, bool bAddLeadingSlash, bool bAddTrailingSlash)
@@ -898,5 +999,3 @@ namespace Electra
 
 
 } // namespace Electra
-
-
diff --git a/Engine/Plugins/Media/ElectraUtil/Source/ElectraBase/Public/Utilities/URLParser.h b/Engine/Plugins/Media/ElectraUtil/Source/ElectraBase/Public/Utilities/URLParser.h
--- a/Engine/Plugins/Media/ElectraUtil/Source/ElectraBase/Public/Utilities/URLParser.h
+++ b/Engine/Plugins/Media/ElectraUtil/Source/ElectraBase/Public/Utilities/URLParser.h
@@ -2,7 +2,9 @@
 
 #pragma once
 
-#include <CoreMinimal.h>
+#include "CoreTypes.h"
+#include "Containers/Array.h"
+#include "Containers/UnrealString.h"
 #include "Utilities/StringHelpers.h"
 
 
@@ -33,16 +35,22 @@ namespace Electra
 		// Returns whether or not this URL is absolute (the scheme not being empty).
 		bool IsAbsolute() const;
 		// Returns the path as individual components. Like GetPath() the components will have escape sequences already decoded.
-		void GetPathComponents(TArray<FString>& OutPathComponents) const;
+		bool GetPathComponents(TArray<FString>& OutPathComponents) const;
 		// Returns the last path component (the "filename").
 		FString GetLastPathComponent() const;
+		// Resolves the given (relative) URL against this one, which will be modified.
+		void ResolveWith(const FURL_RFC3986& Other);
 		// Resolves a relative URL against this one.
 		FURL_RFC3986& ResolveWith(const FString& InChildURL);
 		// Resolves this URL (which should be relative) against the specified URL.
 		FURL_RFC3986& ResolveAgainst(const FString& InParentURL);
 
-		// Appends or prepends additional query parameters.
-		void AddQueryParameters(const FString& InQueryParameters, bool bAppend);
+		void SetScheme(const FString& InNewScheme);
+		void SetHost(const FString& InNewHost);
+		void SetPort(const FString& InNewPort);
+		void SetPath(const FString& InNewPath);
+		void SetQuery(const FString& InNewQuery);
+		void SetFragment(const FString& InNewFragment);
 
 		// Returns if this URL has the same origin as another one as per RFC 6454.
 		bool HasSameOriginAs(const FURL_RFC3986& Other);
@@ -57,10 +65,14 @@ namespace Electra
 		// Returns the given query parameter string as a list of name/value pairs.
 		// The parameter string MUST NOT start with a '?'.
 		static void GetQueryParams(TArray<FQueryParam>& OutQueryParams, const FString& InQueryParameters, bool bPerformUrlDecoding, bool bSameNameReplacesValue=true);
-		
+
 		// Replaces query parameters in the URL.
         void SetQueryParams(const TArray<FQueryParam>& InQueryParams);
 
+		// Updates existing parameters with new value or adds new parameters.
+        void AddOrUpdateQueryParams(const TArray<FQueryParam>& InQueryParams);
+		void AddOrUpdateQueryParams(const FString& InQueryParameters);
+
 		// Decodes %XX escaped sequences into their original characters. Appends to the output. Hence in and out must not be the same.
 		static bool UrlDecode(FString& OutResult, const FString& InUrlToDecode);
 		// Encodes characters not permitted in a URL into %XX escaped sequences. Appends to the output. Hence in and out must not be the same.
@@ -71,6 +83,10 @@ namespace Electra
 		// Returns the standard port for the given scheme. An empty string is returned if none is known.
 		static FString GetStandardPortForScheme(const FString& InScheme, bool bIgnoreCase=true);
 
+		// Returns a string containing the "sub-delims" chars that are permitted in the query string that do not
+		// need to be escaped.
+		static FString GetUrlEncodeSubDelimsChars();
+
 	private:
 		FString Scheme;
 		FString UserInfo;
@@ -79,6 +95,8 @@ namespace Electra
 		FString Path;
 		FString Query;
 		FString Fragment;
+		bool bIsFile = false;
+		bool bIsData = false;
 
 		static void GetPathComponents(TArray<FString>& OutPathComponents, const FString& InPath);
 
@@ -103,7 +121,6 @@ namespace Electra
 		bool ParseFragment(StringHelpers::FStringIterator& it);
 
 		FString GetAuthority() const;
-		void ResolveWith(const FURL_RFC3986& Other);
 		void BuildPathFromSegments(const TArray<FString>& Components, bool bAddLeadingSlash, bool bAddTrailingSlash);
 		void MergePath(const FString& InPathToMerge);
 		void RemoveDotSegments();
