Media: Added player option to select the initial tracks for playback based on language code. This may not be supported with all media players.

#rnx

diff --git a/Engine/Plugins/Media/ElectraPlayer/Source/ElectraPlayerPlugin/Private/ElectraPlayerPlugin.cpp b/Engine/Plugins/Media/ElectraPlayer/Source/ElectraPlayerPlugin/Private/ElectraPlayerPlugin.cpp
--- a/Engine/Plugins/Media/ElectraPlayer/Source/ElectraPlayerPlugin/Private/ElectraPlayerPlugin.cpp
+++ b/Engine/Plugins/Media/ElectraPlayer/Source/ElectraPlayerPlugin/Private/ElectraPlayerPlugin.cpp
@@ -28,8 +28,6 @@ std::atomic<uint32> FElectraPlayerPlugin::NextPlayerUniqueID { 0 };
 namespace ElectraMediaOptions
 {
 	static const FName GetSafeMediaOptions(TEXT("GetSafeMediaOptions"));
-	static const FName InitialAudioLanguage(TEXT("InitialAudioLanguage"));
-	static const FName InitialSubtitleLanguage(TEXT("InitialSubtitleLanguage"));
 	static const FName ElectraNoPreloading(TEXT("ElectraNoPreloading"));
 	static const FName PlaylistProperties(TEXT("playlist_properties"));
 	static const FName ElectraInitialBitrate(TEXT("ElectraInitialBitrate"));
@@ -568,7 +566,25 @@ bool FElectraPlayerPlugin::Open(const FString& Url, const IMediaOptions* Options
 		{
 			LocalPlaystartOptions.TimeOffset = InPlayerOptions->SeekTime;
 		}
-		if (InPlayerOptions->TrackSelection == EMediaPlayerOptionTrackSelectMode::UseTrackOptionIndices)
+		if (InPlayerOptions->TrackSelection == EMediaPlayerOptionTrackSelectMode::UseLanguageCodes)
+		{
+			if (InPlayerOptions->TracksByLanguage.Video.Len())
+			{
+				LocalPlaystartOptions.InitialVideoTrackAttributes.Language_RFC4647 = InPlayerOptions->TracksByLanguage.Video;
+				UE_LOG(LogElectraPlayerPlugin, Verbose, TEXT("[%u] IMediaPlayer::Open: Asking for initial video language \"%s\""), PlayerUniqueID, *InPlayerOptions->TracksByLanguage.Video);
+			}
+			if (InPlayerOptions->TracksByLanguage.Audio.Len())
+			{
+				LocalPlaystartOptions.InitialAudioTrackAttributes.Language_RFC4647 = InPlayerOptions->TracksByLanguage.Audio;
+				UE_LOG(LogElectraPlayerPlugin, Verbose, TEXT("[%u] IMediaPlayer::Open: Asking for initial audio language \"%s\""), PlayerUniqueID, *InPlayerOptions->TracksByLanguage.Audio);
+			}
+			if (InPlayerOptions->TracksByLanguage.Subtitle.Len())
+			{
+				LocalPlaystartOptions.InitialSubtitleTrackAttributes.Language_RFC4647 = InPlayerOptions->TracksByLanguage.Subtitle;
+				UE_LOG(LogElectraPlayerPlugin, Verbose, TEXT("[%u] IMediaPlayer::Open: Asking for initial subtitle language \"%s\""), PlayerUniqueID, *InPlayerOptions->TracksByLanguage.Subtitle);
+			}
+		}
+		else if (InPlayerOptions->TrackSelection == EMediaPlayerOptionTrackSelectMode::UseTrackOptionIndices)
 		{
 			LocalPlaystartOptions.InitialAudioTrackAttributes.TrackIndexOverride = InPlayerOptions->Tracks.Audio;
 			LocalPlaystartOptions.InitialSubtitleTrackAttributes.TrackIndexOverride = InPlayerOptions->Tracks.Subtitle;
@@ -576,18 +592,6 @@ bool FElectraPlayerPlugin::Open(const FString& Url, const IMediaOptions* Options
 		const FVariant* Env = InPlayerOptions->InternalCustomOptions.Find(MediaPlayerOptionValues::Environment());
 		Environment = Env ? Env->GetValue<FName>() : Environment;
 	}
-	FString InitialAudioLanguage = Options->GetMediaOption(ElectraMediaOptions::InitialAudioLanguage, FString());
-	if (InitialAudioLanguage.Len())
-	{
-		LocalPlaystartOptions.InitialAudioTrackAttributes.Language_ISO639 = InitialAudioLanguage;
-		UE_LOG(LogElectraPlayerPlugin, Verbose, TEXT("[%u] IMediaPlayer::Open: Asking for initial audio language \"%s\""), PlayerUniqueID, *InitialAudioLanguage);
-	}
-	FString InitialSubtitleLanguage = Options->GetMediaOption(ElectraMediaOptions::InitialSubtitleLanguage, FString());
-	if (InitialSubtitleLanguage.Len())
-	{
-		LocalPlaystartOptions.InitialSubtitleTrackAttributes.Language_ISO639 = InitialSubtitleLanguage;
-		UE_LOG(LogElectraPlayerPlugin, Verbose, TEXT("[%u] IMediaPlayer::Open: Asking for initial subtitle language \"%s\""), PlayerUniqueID, *InitialSubtitleLanguage);
-	}
 	bool bNoPreloading = Options->GetMediaOption(ElectraMediaOptions::ElectraNoPreloading, (bool)false);
 	if (bNoPreloading)
 	{
diff --git a/Engine/Plugins/Media/ElectraPlayer/Source/ElectraPlayerRuntime/Private/ElectraPlayer.cpp b/Engine/Plugins/Media/ElectraPlayer/Source/ElectraPlayerRuntime/Private/ElectraPlayer.cpp
--- a/Engine/Plugins/Media/ElectraPlayer/Source/ElectraPlayerRuntime/Private/ElectraPlayer.cpp
+++ b/Engine/Plugins/Media/ElectraPlayer/Source/ElectraPlayerRuntime/Private/ElectraPlayer.cpp
@@ -1581,7 +1581,15 @@ FString FElectraPlayer::GetTrackLanguage(EPlayerTrackType TrackType, int32 Track
 	TSharedPtr<Electra::FTrackMetadata, ESPMode::ThreadSafe> Meta = GetTrackStreamMetadata(TrackType, TrackIndex);
 	if (Meta.IsValid())
 	{
-		return Meta->Language;
+		if (TrackType == EPlayerTrackType::Audio)
+		{
+			// Audio does not need to include the script tag (but video does as it could include burned in subtitles)
+			return Meta->LanguageTagRFC5646.Get(true, false, true, false, false, false);
+		}
+		else
+		{
+			return Meta->LanguageTagRFC5646.Get(true, true, true, false, false, false);
+		}
 	}
 	return TEXT("");
 }
@@ -1645,11 +1653,8 @@ bool FElectraPlayer::SelectTrack(EPlayerTrackType TrackType, int32 TrackIndex)
 						TrackAttributes.Kind = Meta->Kind;
 						OutSelectionAttributes.Kind = Meta->Kind;
 					}
-					if (!Meta->Language.IsEmpty())
-					{
-						TrackAttributes.Language_ISO639 = Meta->Language;
-						OutSelectionAttributes.Language_ISO639 = Meta->Language;
-					}
+					TrackAttributes.Language_RFC4647 = Meta->LanguageTagRFC5646.Get(true, true, true, false, false, false);
+					OutSelectionAttributes.Language_RFC4647 = TrackAttributes.Language_RFC4647;
 					TrackAttributes.Codec = Meta->HighestBandwidthCodec.GetCodecName();
 					OutSelectionAttributes.Codec = Meta->HighestBandwidthCodec.GetCodecName();
 
@@ -2135,17 +2140,24 @@ void FElectraPlayer::HandlePlayerEventReceivedPlaylists()
 	CurrentPlayer->AdaptivePlayer->GetTrackMetadata(SubtitleStreamMetaData, Electra::EStreamType::Subtitle);
 	NumTracksSubtitle = SubtitleStreamMetaData.Num();
 
+	// Set the initial video track selection attributes.
+	Electra::FStreamSelectionAttributes InitialVideoAttributes;
+	InitialVideoAttributes.Kind = PlaystartOptions.InitialVideoTrackAttributes.Kind;
+	InitialVideoAttributes.Language_RFC4647 = PlaystartOptions.InitialVideoTrackAttributes.Language_RFC4647;
+	InitialVideoAttributes.OverrideIndex = PlaystartOptions.InitialVideoTrackAttributes.TrackIndexOverride;
+	CurrentPlayer->AdaptivePlayer->SetInitialStreamAttributes(Electra::EStreamType::Video, InitialVideoAttributes);
+
 	// Set the initial audio track selection attributes.
 	Electra::FStreamSelectionAttributes InitialAudioAttributes;
 	InitialAudioAttributes.Kind = PlaystartOptions.InitialAudioTrackAttributes.Kind;
-	InitialAudioAttributes.Language_ISO639 = PlaystartOptions.InitialAudioTrackAttributes.Language_ISO639;
+	InitialAudioAttributes.Language_RFC4647 = PlaystartOptions.InitialAudioTrackAttributes.Language_RFC4647;
 	InitialAudioAttributes.OverrideIndex = PlaystartOptions.InitialAudioTrackAttributes.TrackIndexOverride;
 	CurrentPlayer->AdaptivePlayer->SetInitialStreamAttributes(Electra::EStreamType::Audio, InitialAudioAttributes);
 
 	// Set the initial subtitle track selection attributes.
 	Electra::FStreamSelectionAttributes InitialSubtitleAttributes;
 	InitialSubtitleAttributes.Kind = PlaystartOptions.InitialSubtitleTrackAttributes.Kind;
-	InitialSubtitleAttributes.Language_ISO639 = PlaystartOptions.InitialSubtitleTrackAttributes.Language_ISO639;
+	InitialSubtitleAttributes.Language_RFC4647 = PlaystartOptions.InitialSubtitleTrackAttributes.Language_RFC4647;
 	InitialSubtitleAttributes.OverrideIndex = PlaystartOptions.InitialSubtitleTrackAttributes.TrackIndexOverride;
 	CurrentPlayer->AdaptivePlayer->SetInitialStreamAttributes(Electra::EStreamType::Subtitle, InitialSubtitleAttributes);
 
@@ -2317,6 +2329,12 @@ void FElectraPlayer::HandlePlayerEventBufferingStart(Electra::Metrics::EBufferin
 	// playback will begin and what tracks are available there.
 	if (BufferingReason == Electra::Metrics::EBufferingReason::Initial)
 	{
+		// Mark the track indices as dirty in order to get the current active ones again.
+		// This is necessary since the player may have made a different selection given the
+		// initial track preferences we gave it.
+		bVideoTrackIndexDirty = true;
+		bAudioTrackIndexDirty = true;
+		bSubtitleTrackIndexDirty = true;
 		DeferredEvents.Enqueue(IElectraPlayerAdapterDelegate::EPlayerEvent::TracksChanged);
 	}
 
diff --git a/Engine/Plugins/Media/ElectraPlayer/Source/ElectraPlayerRuntime/Private/Runtime/Demuxer/ParserISO13818-1.cpp b/Engine/Plugins/Media/ElectraPlayer/Source/ElectraPlayerRuntime/Private/Runtime/Demuxer/ParserISO13818-1.cpp
--- a/Engine/Plugins/Media/ElectraPlayer/Source/ElectraPlayerRuntime/Private/Runtime/Demuxer/ParserISO13818-1.cpp
+++ b/Engine/Plugins/Media/ElectraPlayer/Source/ElectraPlayerRuntime/Private/Runtime/Demuxer/ParserISO13818-1.cpp
@@ -1509,7 +1509,9 @@ void FParserISO13818_1::ProcessDescriptor(ElectraDecodersUtil::FElectraBitstream
 					0x85 .. 0xFF Reserved
 				*/
 				FString Language = StringHelpers::ISO_8859_1_ToFString(Lang, 3);
-				OutCodecInfo.SetStreamLanguageCode(Language);
+				BCP47::FLanguageTag LanguageTag;
+				BCP47::ParseRFC5646Tag(LanguageTag, Language);
+				OutCodecInfo.SetStreamLanguageTag(LanguageTag);
 				if (audio_type == 0 || audio_type == 0x80)
 				{
 					break;
diff --git a/Engine/Plugins/Media/ElectraPlayer/Source/ElectraPlayerRuntime/Private/Runtime/Demuxer/ParserISO14496-12.cpp b/Engine/Plugins/Media/ElectraPlayer/Source/ElectraPlayerRuntime/Private/Runtime/Demuxer/ParserISO14496-12.cpp
--- a/Engine/Plugins/Media/ElectraPlayer/Source/ElectraPlayerRuntime/Private/Runtime/Demuxer/ParserISO14496-12.cpp
+++ b/Engine/Plugins/Media/ElectraPlayer/Source/ElectraPlayerRuntime/Private/Runtime/Demuxer/ParserISO14496-12.cpp
@@ -148,7 +148,7 @@ namespace Electra
 				}
 				else
 				{
-					return FString((const char*)SrcBuf.GetData());
+					return FString::ConstructFromPtrSize((const char*)SrcBuf.GetData(), SrcBuf.Num());
 				}
 			};
 
@@ -1513,9 +1513,9 @@ namespace Electra
 			return Timescale;
 		}
 
-		const FString& GetLanguage() const
+		const BCP47::FLanguageTag& GetLanguageTag() const
 		{
-			return Lang639_1;
+			return LanguageTag;
 		}
 
 	private:
@@ -1552,7 +1552,7 @@ namespace Electra
 			Language[1] = (char)(0x60 + ((Value16 & 0x03e0) >> 5));
 			Language[2] = (char)(0x60 + (Value16 & 0x001f));
 			// Try to map the ISO-639-2T language code to the shorter ISO-639-1 code if possible.
-			Lang639_1 = ISO639::MapTo639_1(FString::ConstructFromPtrSize(Language, FMEDIA_STATIC_ARRAY_COUNT(Language)));
+			/*bool bSuccess =*/ BCP47::ParseRFC5646Tag(LanguageTag, FString::ConstructFromPtrSize(Language, FMEDIA_STATIC_ARRAY_COUNT(Language)));
 			RETURN_IF_ERROR(ParseInfo->Reader()->Read(Value16));				// pre_defined (in QuickTime this held 'Quality')
 			return Error;
 		}
@@ -1561,7 +1561,7 @@ namespace Electra
 		uint64		Duration;
 		uint32		Timescale;
 		char		Language[3];
-		FString		Lang639_1;
+		BCP47::FLanguageTag LanguageTag;
 	};
 
 
@@ -3181,6 +3181,11 @@ namespace Electra
 		{
 		}
 
+		const BCP47::FLanguageTag& GetLanguageTag() const
+		{
+			return LanguageTag;
+		}
+
 	private:
 		FMP4BoxELNG() = delete;
 		FMP4BoxELNG(const FMP4BoxELNG&) = delete;
@@ -3191,11 +3196,13 @@ namespace Electra
 			UEMediaError Error = UEMEDIA_ERROR_OK;
 			RETURN_IF_ERROR(FMP4BoxFull::ReadAndParseAttributes(ParseInfo));
 			uint32 BytesRemaining = BoxSize - (ParseInfo->Reader()->GetCurrentReadOffset() - StartOffset);
-			RETURN_IF_ERROR(ParseInfo->Reader()->ReadString(ExtendedLanguage, BytesRemaining, false));
+			RETURN_IF_ERROR(ParseInfo->Reader()->ReadString(ExtendedLanguage, BytesRemaining, true));
+			/*bool bOk =*/ BCP47::ParseRFC5646Tag(LanguageTag, ExtendedLanguage);
 			return Error;
 		}
 
 	private:
+		BCP47::FLanguageTag LanguageTag;
 		FString ExtendedLanguage;
 	};
 
@@ -6195,7 +6202,7 @@ namespace Electra
 			const TArray<uint8> GetCodecSpecificDataRAW() const override;
 			const FStreamCodecInformation GetCodecInformation() const override;
 			const FBitrateInfo GetBitrateInfo() const override;
-			const FString& GetLanguage() const override;
+			const BCP47::FLanguageTag& GetLanguageTag() const override;
 			bool GetEncryptionInfo(FEncryptionInfo& OutEncryptionInfo) const override;
 			void GetPSSHBoxes(TArray<TArray<uint8>>& OutBoxes, bool bFromMOOV, bool bFromMOOF) const override;
 			void GetPRFTBoxes(TArray<ITrack::FProducerReferenceTime>& OutBoxes) const override;
@@ -6206,6 +6213,7 @@ namespace Electra
 			const FMP4BoxTKHD* TKHDBox = nullptr;
 			const FMP4Box*	   MDIABox = nullptr;
 			const FMP4BoxMDHD* MDHDBox = nullptr;
+			const FMP4BoxELNG* ELNGBox = nullptr;
 			const FMP4BoxHDLR* HDLRBox = nullptr;
 			const FMP4Box*	   STBLBox = nullptr;
 			const FMP4BoxSTSD* STSDBox = nullptr;
@@ -6528,10 +6536,10 @@ namespace Electra
 		return BitrateInfo;
 	}
 
-	const FString& FParserISO14496_12::FTrack::GetLanguage() const
+	const BCP47::FLanguageTag& FParserISO14496_12::FTrack::GetLanguageTag() const
 	{
-		static FString Undefined(TEXT("und"));
-		return MDHDBox ? MDHDBox->GetLanguage() : Undefined;
+		static BCP47::FLanguageTag Undefined { .FullLanguage=FString(TEXT("und")), .PrimaryLanguage=FString(TEXT("und"))};
+		return ELNGBox ? ELNGBox->GetLanguageTag() : MDHDBox ? MDHDBox->GetLanguageTag() : Undefined;
 	}
 
 	bool FParserISO14496_12::FTrack::GetEncryptionInfo(FEncryptionInfo& OutEncryptionInfo) const
@@ -7513,7 +7521,7 @@ namespace Electra
 							Track->CodecSpecificDataAVC = AVCCBox->GetDecoderConfigurationRecord().GetValue();
 							Track->CodecInformation.SetStreamType(EStreamType::Video);
 							Track->CodecInformation.SetCodec(FStreamCodecInformation::ECodec::H264);
-							Track->CodecInformation.SetStreamLanguageCode(Track->GetLanguage());
+							Track->CodecInformation.SetStreamLanguageTag(Track->GetLanguageTag());
 							if (Track->CodecSpecificDataAVC.GetSequenceParameterSets().IsEmpty())
 							{
 								if (InAvcType == 3)
@@ -7636,7 +7644,7 @@ namespace Electra
 							Track->CodecSpecificDataHEVC = HVCCBox->GetDecoderConfigurationRecord().GetValue();
 							Track->CodecInformation.SetStreamType(EStreamType::Video);
 							Track->CodecInformation.SetCodec(FStreamCodecInformation::ECodec::H265);
-							Track->CodecInformation.SetStreamLanguageCode(Track->GetLanguage());
+							Track->CodecInformation.SetStreamLanguageTag(Track->GetLanguageTag());
 							Track->CodecInformation.SetProfileSpace(Track->CodecSpecificDataHEVC.GetGeneralProfileSpace());
 							Track->CodecInformation.SetProfileTier(Track->CodecSpecificDataHEVC.GetGeneralTierFlag());
 							Track->CodecInformation.SetProfile(Track->CodecSpecificDataHEVC.GetGeneralProfileIDC());
@@ -7752,7 +7760,7 @@ namespace Electra
 				{
 					Track->CodecInformation.SetStreamType(EStreamType::Audio);
 					Track->CodecInformation.SetCodec(FStreamCodecInformation::ECodec::AAC);
-					Track->CodecInformation.SetStreamLanguageCode(Track->GetLanguage());
+					Track->CodecInformation.SetStreamLanguageTag(Track->GetLanguageTag());
 
 					ElectraDecodersUtil::MPEG::FAACDecoderConfigurationRecord ConfigRecord;
 					if (ConfigRecord.ParseFrom(Track->CodecSpecificDataMP4A.GetCodecSpecificData().GetData(), Track->CodecSpecificDataMP4A.GetCodecSpecificData().Num()))
@@ -7795,7 +7803,7 @@ namespace Electra
 					Track->CodecInformation.SetCodec(FStreamCodecInformation::ECodec::Audio4CC);
 					Track->CodecInformation.SetCodec4CC(Utils::Make4CC('m','p','g','a'));
 					Track->CodecInformation.SetProfile(1);
-					Track->CodecInformation.SetStreamLanguageCode(Track->GetLanguage());
+					Track->CodecInformation.SetStreamLanguageTag(Track->GetLanguageTag());
 					Track->CodecInformation.SetCodecSpecifierRFC6381(FString::Printf(TEXT("mp4a.6b")));
 					Track->CodecInformation.SetSamplingRate((int32) AudioSampleEntry->GetSampleRate());
 					Track->CodecInformation.SetNumberOfChannels(AudioSampleEntry->GetChannelCount());
@@ -7905,7 +7913,7 @@ namespace Electra
 			Track->CodecInformation.SetStreamType(EStreamType::Subtitle);
 			Track->CodecInformation.SetCodec(FStreamCodecInformation::ECodec::TX3G);
 			Track->CodecInformation.SetCodecSpecificData(Track->CodecSpecificDataRAW);
-			Track->CodecInformation.SetStreamLanguageCode(Track->GetLanguage());
+			Track->CodecInformation.SetStreamLanguageTag(Track->GetLanguageTag());
 			Track->CodecInformation.SetCodecSpecifierRFC6381(TEXT("tx3g"));
 			Track->CodecInformation.SetResolution(FStreamCodecInformation::FResolution(SubtitleSampleEntry->GetWidth(), SubtitleSampleEntry->GetHeight()));
 			Track->CodecInformation.SetTranslation(FStreamCodecInformation::FTranslation(SubtitleSampleEntry->GetTranslationX(), SubtitleSampleEntry->GetTranslationY()));
@@ -7926,7 +7934,7 @@ namespace Electra
 			Track->CodecInformation.SetStreamType(EStreamType::Subtitle);
 			Track->CodecInformation.SetCodec(FStreamCodecInformation::ECodec::WebVTT);
 			Track->CodecInformation.SetCodecSpecificData(Track->CodecSpecificDataRAW);
-			Track->CodecInformation.SetStreamLanguageCode(Track->GetLanguage());
+			Track->CodecInformation.SetStreamLanguageTag(Track->GetLanguageTag());
 			Track->CodecInformation.SetCodecSpecifierRFC6381(TEXT("wvtt"));
 			Track->CodecInformation.SetResolution(FStreamCodecInformation::FResolution(SubtitleSampleEntry->GetWidth(), SubtitleSampleEntry->GetHeight()));
 			Track->CodecInformation.SetTranslation(FStreamCodecInformation::FTranslation(SubtitleSampleEntry->GetTranslationX(), SubtitleSampleEntry->GetTranslationY()));
@@ -7947,7 +7955,7 @@ namespace Electra
 			Track->CodecInformation.SetStreamType(EStreamType::Subtitle);
 			Track->CodecInformation.SetCodec(FStreamCodecInformation::ECodec::TTML);
 			Track->CodecInformation.SetCodecSpecificData(Track->CodecSpecificDataRAW);
-			Track->CodecInformation.SetStreamLanguageCode(Track->GetLanguage());
+			Track->CodecInformation.SetStreamLanguageTag(Track->GetLanguageTag());
 			Track->CodecInformation.SetCodecSpecifierRFC6381(TEXT("stpp"));
 			Track->CodecInformation.SetResolution(FStreamCodecInformation::FResolution(SubtitleSampleEntry->GetWidth(), SubtitleSampleEntry->GetHeight()));
 			Track->CodecInformation.SetTranslation(FStreamCodecInformation::FTranslation(SubtitleSampleEntry->GetTranslationX(), SubtitleSampleEntry->GetTranslationY()));
@@ -8057,7 +8065,7 @@ namespace Electra
 					CodecInformation.SetStreamType(EStreamType::Video);
 					CodecInformation.SetCodec(FStreamCodecInformation::ECodec::Video4CC);
 					CodecInformation.SetCodec4CC(FRMABox ? FRMABox->GetDataFormat() : VisualSampleEntry->GetType());
-					CodecInformation.SetStreamLanguageCode(Track->GetLanguage());
+					CodecInformation.SetStreamLanguageTag(Track->GetLanguageTag());
 					CodecInformation.SetResolution(FStreamCodecInformation::FResolution(VisualSampleEntry->GetWidth(), VisualSampleEntry->GetHeight()));
 					CodecInformation.SetCodecSpecifierRFC6381(Utils::Printable4CC(CodecInformation.GetCodec4CC()));
 					// Try to get the framerate
@@ -8132,7 +8140,7 @@ namespace Electra
 					CodecInformation.SetStreamType(EStreamType::Audio);
 					CodecInformation.SetCodec(FStreamCodecInformation::ECodec::Audio4CC);
 					CodecInformation.SetCodec4CC(FRMABox ? FRMABox->GetDataFormat() : AudioSampleEntry->GetType());
-					CodecInformation.SetStreamLanguageCode(Track->GetLanguage());
+					CodecInformation.SetStreamLanguageTag(Track->GetLanguageTag());
 					CodecInformation.SetSamplingRate((int32) AudioSampleEntry->GetSampleRate());
 					CodecInformation.SetNumberOfChannels(AudioSampleEntry->GetChannelCount());
 					CodecInformation.SetCodecSpecifierRFC6381(Utils::Printable4CC(CodecInformation.GetCodec4CC()));
@@ -8272,6 +8280,7 @@ namespace Electra
 						{
 							return UEMEDIA_ERROR_FORMAT_ERROR;
 						}
+						Track->ELNGBox = static_cast<const FMP4BoxELNG*>(Track->MDIABox->GetBoxPath(FMP4Box::kBox_elng));
 						Track->UDTABox = static_cast<const FMP4BoxUDTA*>(Box->GetBoxPath(FMP4Box::kBox_udta));
 						Track->TREFBox = static_cast<const FMP4BoxTREF*>(Box->GetBoxPath(FMP4Box::kBox_tref));
 						// TODO: Check the dinf->dref->url box to not reference any external media!
diff --git a/Engine/Plugins/Media/ElectraPlayer/Source/ElectraPlayerRuntime/Private/Runtime/Demuxer/ParserISO14496-12.h b/Engine/Plugins/Media/ElectraPlayer/Source/ElectraPlayerRuntime/Private/Runtime/Demuxer/ParserISO14496-12.h
--- a/Engine/Plugins/Media/ElectraPlayer/Source/ElectraPlayerRuntime/Private/Runtime/Demuxer/ParserISO14496-12.h
+++ b/Engine/Plugins/Media/ElectraPlayer/Source/ElectraPlayerRuntime/Private/Runtime/Demuxer/ParserISO14496-12.h
@@ -227,7 +227,7 @@ namespace Electra
 			virtual const TArray<uint8> GetCodecSpecificDataRAW() const = 0;
 			virtual const FStreamCodecInformation GetCodecInformation() const = 0;
 			virtual const FBitrateInfo GetBitrateInfo() const = 0;
-			virtual const FString& GetLanguage() const = 0;
+			virtual const BCP47::FLanguageTag& GetLanguageTag() const = 0;
 			virtual bool GetEncryptionInfo(FEncryptionInfo& OutEncryptionInfo) const = 0;
 			virtual void GetPSSHBoxes(TArray<TArray<uint8>>& OutBoxes, bool bFromMOOV, bool bFromMOOF) const = 0;
 			virtual void GetPRFTBoxes(TArray<FProducerReferenceTime>& OutBoxes) const = 0;
diff --git a/Engine/Plugins/Media/ElectraPlayer/Source/ElectraPlayerRuntime/Private/Runtime/Demuxer/ParserMKV.cpp b/Engine/Plugins/Media/ElectraPlayer/Source/ElectraPlayerRuntime/Private/Runtime/Demuxer/ParserMKV.cpp
--- a/Engine/Plugins/Media/ElectraPlayer/Source/ElectraPlayerRuntime/Private/Runtime/Demuxer/ParserMKV.cpp
+++ b/Engine/Plugins/Media/ElectraPlayer/Source/ElectraPlayerRuntime/Private/Runtime/Demuxer/ParserMKV.cpp
@@ -3,11 +3,11 @@
 #include "ParserMKV.h"
 
 #include "Misc/Build.h"
-
 #include "BufferedDataReader.h"
 
 #include "StreamTypes.h"
 #include "Utilities/ISO639-Map.h"
+#include "Utilities/BCP47-Helpers.h"
 #include "InfoLog.h"
 #include "Player/PlayerSessionServices.h"
 #include "Player/PlayerStreamFilter.h"
@@ -1353,7 +1353,12 @@ namespace Electra
 		public:
 			FMKVTrackEntry(uint32 InID, int64 InByteOffset, int64 InByteSize) : FMKVEBMLElement(InID, InByteOffset, InByteSize) { }
 			virtual ~FMKVTrackEntry() = default;
-			virtual EParseResult ParseElement(IMKVFetcher* InReader) { return ParseElementList(Elements, InReader); }
+			EParseResult ParseElement(IMKVFetcher* InReader)  override
+			{
+				EParseResult Result = ParseElementList(Elements, InReader);
+				BCP47::ParseRFC5646Tag(LanguageTag, GetLanguage());
+				return Result;
+			}
 
 			uint64 GetTrackNumber() const
 			{ return TrackNumber; }
@@ -1385,6 +1390,8 @@ namespace Electra
 			{ return Name; }
 			const FString& GetLanguage() const
 			{ return LanguageBCP47.Len() ? LanguageBCP47 : Language; }
+			const BCP47::FLanguageTag& GetLanguageTag() const
+			{ return LanguageTag; }
 			const FString& GetCodecID() const
 			{ return CodecID; }
 			const TArray<uint8>& GetCodecPrivate() const
@@ -1406,6 +1413,7 @@ namespace Electra
 			FString Name;
 			FString Language {TEXT("eng")};
 			FString LanguageBCP47;
+			BCP47::FLanguageTag LanguageTag;
 			FString CodecID;
 			FString CodecName;
 			double TrackTimestampScale = 1.0;
@@ -2119,6 +2127,8 @@ namespace Electra
 			{ return QuickInfo.CodecInfo; }
 			const FString GetLanguage() const override
 			{ return Track->GetLanguage(); }
+			const BCP47::FLanguageTag GetLanguageTag() const override
+			{ return Track->GetLanguageTag(); }
 			IParserMKV::ICueIterator* CreateCueIterator() const override;
 			const FQuickInfo GetQuickInfo() const
 			{ return QuickInfo;}
@@ -2743,7 +2753,7 @@ namespace Electra
 			OutCodecInformation.SetCodec(FStreamCodecInformation::ECodec::H264);
 			OutCodecInformation.SetCodecSpecificData(CodecSpecificDataAVC.GetCodecSpecificData());
 			OutCodecInformation.SetDecoderConfigRecord(InFromTrack->GetCodecPrivate());
-			OutCodecInformation.SetStreamLanguageCode(InFromTrack->GetLanguage());
+			OutCodecInformation.SetStreamLanguageTag(InFromTrack->GetLanguageTag());
 			if (CodecSpecificDataAVC.GetSequenceParameterSets().Num())
 			{
 				ElectraDecodersUtil::MPEG::H264::FSequenceParameterSet sps;
@@ -2825,7 +2835,7 @@ namespace Electra
 			OutCodecInformation.SetCodec(FStreamCodecInformation::ECodec::H265);
 			OutCodecInformation.SetCodecSpecificData(CodecSpecificDataHEVC.GetCodecSpecificData());
 			OutCodecInformation.SetDecoderConfigRecord(InFromTrack->GetCodecPrivate());
-			OutCodecInformation.SetStreamLanguageCode(InFromTrack->GetLanguage());
+			OutCodecInformation.SetStreamLanguageTag(InFromTrack->GetLanguageTag());
 			if (CodecSpecificDataHEVC.GetParsedSequenceParameterSets().IsEmpty())
 			{
 				return false;
@@ -2922,7 +2932,7 @@ namespace Electra
 
 			OutCodecInformation.SetStreamType(EStreamType::Video);
 			OutCodecInformation.SetMimeType(TEXT("video/VP8"));
-			OutCodecInformation.SetStreamLanguageCode(InFromTrack->GetLanguage());
+			OutCodecInformation.SetStreamLanguageTag(InFromTrack->GetLanguageTag());
 			OutCodecInformation.SetCodec(FStreamCodecInformation::ECodec::Video4CC);
 			OutCodecInformation.SetCodec4CC(Utils::Make4CC('v','p','0','8'));
 			FString crfc = FString::Printf(TEXT("vp08.%02d.%02d.%02d.%02d.%02d.%02d.%02d.%02d"),
@@ -3066,7 +3076,7 @@ namespace Electra
 
 			OutCodecInformation.SetStreamType(EStreamType::Video);
 			OutCodecInformation.SetMimeType(TEXT("video/VP9"));
-			OutCodecInformation.SetStreamLanguageCode(InFromTrack->GetLanguage());
+			OutCodecInformation.SetStreamLanguageTag(InFromTrack->GetLanguageTag());
 			OutCodecInformation.SetCodec(FStreamCodecInformation::ECodec::Video4CC);
 			OutCodecInformation.SetCodec4CC(Utils::Make4CC('v','p','0','9'));
 			FString crfc = FString::Printf(TEXT("vp09.%02d.%02d.%02d.%02d.%02d.%02d.%02d.%02d"),
@@ -3132,7 +3142,7 @@ namespace Electra
 			OutCodecInformation.SetMimeType(TEXT("audio/mp4"));
 			OutCodecInformation.SetCodec(FStreamCodecInformation::ECodec::AAC);
 			OutCodecInformation.SetCodecSpecificData(AudioSpecificConfig.GetCodecSpecificData());
-			OutCodecInformation.SetStreamLanguageCode(InFromTrack->GetLanguage());
+			OutCodecInformation.SetStreamLanguageTag(InFromTrack->GetLanguageTag());
 			OutCodecInformation.SetCodecSpecifierRFC6381(FString::Printf(TEXT("mp4a.40.%d"), AudioSpecificConfig.ExtAOT ? AudioSpecificConfig.ExtAOT : AudioSpecificConfig.AOT));
 			OutCodecInformation.SetSamplingRate(AudioSpecificConfig.ExtSamplingFrequency ? AudioSpecificConfig.ExtSamplingFrequency : AudioSpecificConfig.SamplingRate);
 			OutCodecInformation.SetChannelConfiguration(AudioSpecificConfig.ChannelConfiguration);
@@ -3205,7 +3215,7 @@ namespace Electra
 			OutCodecInformation.SetCodec(FStreamCodecInformation::ECodec::Audio4CC);
 			OutCodecInformation.SetCodec4CC(Utils::Make4CC('O','p','u','s'));
 			OutCodecInformation.SetCodecSpecificData(OpusHead);
-			OutCodecInformation.SetStreamLanguageCode(InFromTrack->GetLanguage());
+			OutCodecInformation.SetStreamLanguageTag(InFromTrack->GetLanguageTag());
 			OutCodecInformation.SetCodecSpecifierRFC6381(TEXT("Opus"));
 			OutCodecInformation.SetSamplingRate(Audio->GetOutputSampleRate());
 			OutCodecInformation.SetNumberOfChannels(Audio->GetNumberOfChannels());
@@ -3255,7 +3265,7 @@ namespace Electra
 			OutCodecInformation.SetCodec(FStreamCodecInformation::ECodec::Audio4CC);
 			OutCodecInformation.SetCodec4CC(Utils::Make4CC('f','L','a','C'));
 			OutCodecInformation.SetCodecSpecificData(FlacHead);
-			OutCodecInformation.SetStreamLanguageCode(InFromTrack->GetLanguage());
+			OutCodecInformation.SetStreamLanguageTag(InFromTrack->GetLanguageTag());
 			OutCodecInformation.SetCodecSpecifierRFC6381(TEXT("flac"));
 			OutCodecInformation.SetSamplingRate(Audio->GetOutputSampleRate());
 			OutCodecInformation.SetNumberOfChannels(Audio->GetNumberOfChannels());
@@ -3288,7 +3298,7 @@ namespace Electra
 				LexFromString(LayerValue, *Layer);
 			}
 			OutCodecInformation.SetProfileLevel(LayerValue);
-			OutCodecInformation.SetStreamLanguageCode(InFromTrack->GetLanguage());
+			OutCodecInformation.SetStreamLanguageTag(InFromTrack->GetLanguageTag());
 			OutCodecInformation.SetCodecSpecifierRFC6381(TEXT("mp4a.6b"));
 			OutCodecInformation.SetSamplingRate(Audio->GetOutputSampleRate());
 			OutCodecInformation.SetNumberOfChannels(Audio->GetNumberOfChannels());
diff --git a/Engine/Plugins/Media/ElectraPlayer/Source/ElectraPlayerRuntime/Private/Runtime/Demuxer/ParserMKV.h b/Engine/Plugins/Media/ElectraPlayer/Source/ElectraPlayerRuntime/Private/Runtime/Demuxer/ParserMKV.h
--- a/Engine/Plugins/Media/ElectraPlayer/Source/ElectraPlayerRuntime/Private/Runtime/Demuxer/ParserMKV.h
+++ b/Engine/Plugins/Media/ElectraPlayer/Source/ElectraPlayerRuntime/Private/Runtime/Demuxer/ParserMKV.h
@@ -184,6 +184,7 @@ namespace Electra
 			virtual const TArray<uint8>& GetCodecSpecificData() const = 0;
 			virtual const FStreamCodecInformation& GetCodecInformation() const = 0;
 			virtual const FString GetLanguage() const = 0;
+			virtual const BCP47::FLanguageTag GetLanguageTag() const = 0;
 
 			virtual ICueIterator* CreateCueIterator() const = 0;
 		};
diff --git a/Engine/Plugins/Media/ElectraPlayer/Source/ElectraPlayerRuntime/Private/Runtime/Player/AdaptiveStreamingPlayer.cpp b/Engine/Plugins/Media/ElectraPlayer/Source/ElectraPlayerRuntime/Private/Runtime/Player/AdaptiveStreamingPlayer.cpp
--- a/Engine/Plugins/Media/ElectraPlayer/Source/ElectraPlayerRuntime/Private/Runtime/Player/AdaptiveStreamingPlayer.cpp
+++ b/Engine/Plugins/Media/ElectraPlayer/Source/ElectraPlayerRuntime/Private/Runtime/Player/AdaptiveStreamingPlayer.cpp
@@ -10,7 +10,7 @@
 #include "Player/PlayerLicenseKey.h"
 #include "Player/AdaptivePlayerOptionKeynames.h"
 #include "Utilities/Utilities.h"
-#include "Utilities/ISO639-Map.h"
+#include "Utilities/BCP47-Helpers.h"
 #include "Player/DRM/DRMManager.h"
 
 #include "HAL/LowLevelMemTracker.h"
@@ -833,7 +833,15 @@ void FAdaptiveStreamingPlayer::GetSelectedTrackAttributes(FStreamSelectionAttrib
 		if (BufferInfo.IsValid())
 		{
 			OutAttributes.Kind = BufferInfo->Kind;
-			OutAttributes.Language_ISO639 = BufferInfo->Language;
+			FString Lang = BufferInfo->LanguageTag.Get();
+			if (!Lang.IsEmpty())
+			{
+				OutAttributes.Language_RFC4647 = MoveTemp(Lang);
+			}
+			else
+			{
+				OutAttributes.Language_RFC4647.Reset();
+			}
 			OutAttributes.Codec = BufferInfo->Codec;
 			OutAttributes.OverrideIndex = BufferInfo->HardIndex;
 		}
@@ -896,9 +904,45 @@ void FAdaptiveStreamingPlayer::SetMaxResolution(int32 MaxWidth, int32 MaxHeight)
 }
 
 
-void FAdaptiveStreamingPlayer::SetInitialStreamAttributes(EStreamType StreamType, const FStreamSelectionAttributes& InitialSelection)
+void FAdaptiveStreamingPlayer::SetInitialStreamAttributes(EStreamType StreamType, FStreamSelectionAttributes InitialSelection)
 {
-	WorkerThread.SendInitialStreamAttributeMessage(StreamType, InitialSelection);
+	bool bInitiallyDisabled = InitialSelection.OverrideIndex.IsSet() && InitialSelection.OverrideIndex.GetValue() < 0;
+	if (bInitiallyDisabled)
+	{
+		InitialSelection.ClearOverrideIndex();
+	}
+	switch(StreamType)
+	{
+		case EStreamType::Video:
+		{
+			StreamSelectionAttributesVid = InitialSelection;
+			break;
+		}
+		case EStreamType::Audio:
+		{
+			StreamSelectionAttributesAud = InitialSelection;
+			break;
+		}
+		case EStreamType::Subtitle:
+		{
+			StreamSelectionAttributesTxt = InitialSelection;
+			// If there is an initial subtitle selection (whether it can be fulfilled or not)
+			// we enable them. They were turned off implicitly in the constructor.
+			if (StreamSelectionAttributesTxt.IsSet())
+			{
+				SelectedStreamAttributesTxt.Select();
+			}
+			break;
+		}
+		default:
+		{
+			break;
+		}
+	}
+	if (bInitiallyDisabled)
+	{
+		InternalDeselectStream(StreamType);
+	}
 }
 
 #if 0
@@ -1540,50 +1584,11 @@ bool FAdaptiveStreamingPlayer::InternalHandleThreadMessages()
 				break;
 			}
 
-			case FWorkerThreadMessages::FMessage::EType::InitialStreamAttributes:
-			{
-				FWorkerThreadMessages::FMessage::FInitialStreamSelect& ev = msg.Data.Get<FWorkerThreadMessages::FMessage::FInitialStreamSelect>();
-				// Map the language in case it is not yet.
-				if (ev.InitialSelection.Language_ISO639.IsSet())
-				{
-					ev.InitialSelection.Language_ISO639 = ISO639::MapTo639_1(ev.InitialSelection.Language_ISO639.GetValue());
-				}
-				bool bInitiallyDisabled = ev.InitialSelection.OverrideIndex.IsSet() && ev.InitialSelection.OverrideIndex.GetValue() < 0;
-				if (bInitiallyDisabled)
-				{
-					ev.InitialSelection.ClearOverrideIndex();
-				}
-				switch(ev.StreamType)
-				{
-					case EStreamType::Video:
-						StreamSelectionAttributesVid = ev.InitialSelection;
-						break;
-					case EStreamType::Audio:
-						StreamSelectionAttributesAud = ev.InitialSelection;
-						break;
-					case EStreamType::Subtitle:
-						StreamSelectionAttributesTxt = ev.InitialSelection;
-						break;
-					default:
-						break;
-				}
-				if (bInitiallyDisabled)
-				{
-					InternalDeselectStream(ev.StreamType);
-				}
-				break;
-			}
-
 			case FWorkerThreadMessages::FMessage::EType::SelectTrackByAttributes:
 			{
 				if (!bIsClosing)
 				{
 					FWorkerThreadMessages::FMessage::FMetadataTrackSelection& ev = msg.Data.Get<FWorkerThreadMessages::FMessage::FMetadataTrackSelection>();
-					// Map the language in case it is not yet.
-					if (ev.TrackAttributes.Language_ISO639.IsSet())
-					{
-						ev.TrackAttributes.Language_ISO639 = ISO639::MapTo639_1(ev.TrackAttributes.Language_ISO639.GetValue());
-					}
 					switch(ev.StreamType)
 					{
 						case EStreamType::Video:
@@ -2735,26 +2740,26 @@ void FAdaptiveStreamingPlayer::InternalHandlePendingStartRequest(const FTimeValu
 						TSharedPtrTS<FBufferSourceInfo> BufferSourceInfoTxt = CurrentPlayPeriodText->GetSelectedStreamBufferSourceInfo(EStreamType::Subtitle);
 						if (BufferSourceInfoVid.IsValid())
 						{
-							SelectedStreamAttributesVid.UpdateWith(BufferSourceInfoVid->Kind, BufferSourceInfoVid->Language, BufferSourceInfoVid->Codec, BufferSourceInfoVid->HardIndex);
+							SelectedStreamAttributesVid.UpdateWith(BufferSourceInfoVid->Kind, BufferSourceInfoVid->LanguageTag, BufferSourceInfoVid->Codec, BufferSourceInfoVid->HardIndex);
 							if (ManifestType != EMediaFormatType::ISOBMFF && ManifestType != EMediaFormatType::MKV)
 							{
-								StreamSelectionAttributesVid.UpdateIfOverrideSet(BufferSourceInfoVid->Kind, BufferSourceInfoVid->Language, BufferSourceInfoVid->Codec);
+								StreamSelectionAttributesVid.UpdateIfOverrideSet(BufferSourceInfoVid->Kind, BufferSourceInfoVid->LanguageTag, BufferSourceInfoVid->Codec);
 							}
 						}
 						if (BufferSourceInfoAud.IsValid())
 						{
-							SelectedStreamAttributesAud.UpdateWith(BufferSourceInfoAud->Kind, BufferSourceInfoAud->Language, BufferSourceInfoAud->Codec, BufferSourceInfoAud->HardIndex);
+							SelectedStreamAttributesAud.UpdateWith(BufferSourceInfoAud->Kind, BufferSourceInfoAud->LanguageTag, BufferSourceInfoAud->Codec, BufferSourceInfoAud->HardIndex);
 							if (ManifestType != EMediaFormatType::ISOBMFF && ManifestType != EMediaFormatType::MKV)
 							{
-								StreamSelectionAttributesAud.UpdateIfOverrideSet(BufferSourceInfoAud->Kind, BufferSourceInfoAud->Language, BufferSourceInfoAud->Codec);
+								StreamSelectionAttributesAud.UpdateIfOverrideSet(BufferSourceInfoAud->Kind, BufferSourceInfoAud->LanguageTag, BufferSourceInfoAud->Codec);
 							}
 						}
 						if (BufferSourceInfoTxt.IsValid())
 						{
-							SelectedStreamAttributesTxt.UpdateWith(BufferSourceInfoTxt->Kind, BufferSourceInfoTxt->Language, BufferSourceInfoTxt->Codec, BufferSourceInfoTxt->HardIndex);
+							SelectedStreamAttributesTxt.UpdateWith(BufferSourceInfoTxt->Kind, BufferSourceInfoTxt->LanguageTag, BufferSourceInfoTxt->Codec, BufferSourceInfoTxt->HardIndex);
 							if (ManifestType != EMediaFormatType::ISOBMFF && ManifestType != EMediaFormatType::MKV)
 							{
-								StreamSelectionAttributesTxt.UpdateIfOverrideSet(BufferSourceInfoTxt->Kind, BufferSourceInfoTxt->Language, BufferSourceInfoTxt->Codec);
+								StreamSelectionAttributesTxt.UpdateIfOverrideSet(BufferSourceInfoTxt->Kind, BufferSourceInfoTxt->LanguageTag, BufferSourceInfoTxt->Codec);
 							}
 						}
 
@@ -3369,7 +3374,7 @@ void FAdaptiveStreamingPlayer::HandlePendingMediaSegmentRequests()
 							{
 								case EStreamType::Video:
 								{
-									SelectedStreamAttributesVid.UpdateWith(BufferSourceInfo->Kind, BufferSourceInfo->Language, BufferSourceInfo->Codec, BufferSourceInfo->HardIndex);
+									SelectedStreamAttributesVid.UpdateWith(BufferSourceInfo->Kind, BufferSourceInfo->LanguageTag, BufferSourceInfo->Codec, BufferSourceInfo->HardIndex);
 									if (ManifestType != EMediaFormatType::ISOBMFF && ManifestType != EMediaFormatType::MKV)
 									{
 										StreamSelectionAttributesVid.ClearOverrideIndex();
@@ -3378,7 +3383,7 @@ void FAdaptiveStreamingPlayer::HandlePendingMediaSegmentRequests()
 								}
 								case EStreamType::Audio:
 								{
-									SelectedStreamAttributesAud.UpdateWith(BufferSourceInfo->Kind, BufferSourceInfo->Language, BufferSourceInfo->Codec, BufferSourceInfo->HardIndex);
+									SelectedStreamAttributesAud.UpdateWith(BufferSourceInfo->Kind, BufferSourceInfo->LanguageTag, BufferSourceInfo->Codec, BufferSourceInfo->HardIndex);
 									if (ManifestType != EMediaFormatType::ISOBMFF && ManifestType != EMediaFormatType::MKV)
 									{
 										StreamSelectionAttributesAud.ClearOverrideIndex();
@@ -3387,7 +3392,7 @@ void FAdaptiveStreamingPlayer::HandlePendingMediaSegmentRequests()
 								}
 								case EStreamType::Subtitle:
 								{
-									SelectedStreamAttributesTxt.UpdateWith(BufferSourceInfo->Kind, BufferSourceInfo->Language, BufferSourceInfo->Codec, BufferSourceInfo->HardIndex);
+									SelectedStreamAttributesTxt.UpdateWith(BufferSourceInfo->Kind, BufferSourceInfo->LanguageTag, BufferSourceInfo->Codec, BufferSourceInfo->HardIndex);
 									if (ManifestType != EMediaFormatType::ISOBMFF && ManifestType != EMediaFormatType::MKV)
 									{
 										StreamSelectionAttributesTxt.ClearOverrideIndex();
@@ -3672,7 +3677,7 @@ void FAdaptiveStreamingPlayer::InternalHandleSegmentTrackChanges(const FTimeValu
 					OutSelectionAttributes = *InPendingSelection;
 					bChangeMade = true;
 				}
-				InOutSelectedAttributes.UpdateWith(BufferSourceInfo->Kind, BufferSourceInfo->Language, BufferSourceInfo->Codec, BufferSourceInfo->HardIndex);
+				InOutSelectedAttributes.UpdateWith(BufferSourceInfo->Kind, BufferSourceInfo->LanguageTag, BufferSourceInfo->Codec, BufferSourceInfo->HardIndex);
 				InOutSelectedAttributes.Select();
 				InPendingSelection.Reset();
 			}
@@ -3729,7 +3734,7 @@ void FAdaptiveStreamingPlayer::InternalHandleSegmentTrackChanges(const FTimeValu
 			}
 		};
 
-		PendingTrackSelectionVid.Reset();	// Ignore video changes for now.
+		HandleElementaryStreamTrackSelection(StreamSelectionAttributesVid, SelectedStreamAttributesVid, EStreamType::Video, PendingTrackSelectionVid, CurrentPlayPeriodVideo);
 		HandleElementaryStreamTrackSelection(StreamSelectionAttributesAud, SelectedStreamAttributesAud, EStreamType::Audio, PendingTrackSelectionAud, CurrentPlayPeriodAudio);
 		if (ManifestType == EMediaFormatType::DASH)
 		{
@@ -3861,16 +3866,24 @@ void FAdaptiveStreamingPlayer::UpdatePeriodStreamBufferSourceInfo(TSharedPtrTS<I
 					switch(InStreamType)
 					{
 						case EStreamType::Video:
+						{
 							P.BufferSourceInfoVid = InBufferSourceInfo;
 							break;
+						}
 						case EStreamType::Audio:
+						{
 							P.BufferSourceInfoAud = InBufferSourceInfo;
 							break;
+						}
 						case EStreamType::Subtitle:
+						{
 							P.BufferSourceInfoTxt = InBufferSourceInfo;
 							break;
+						}
 						default:
+						{
 							break;
+						}
 					}
 				}
 			}
@@ -3903,16 +3916,24 @@ TSharedPtrTS<FBufferSourceInfo> FAdaptiveStreamingPlayer::GetStreamBufferInfoAtT
 			switch(InStreamType)
 			{
 				case EStreamType::Video:
+				{
 					BufferSourceInfo = ActivePeriods[i].BufferSourceInfoVid;
 					break;
+				}
 				case EStreamType::Audio:
+				{
 					BufferSourceInfo = ActivePeriods[i].BufferSourceInfoAud;
 					break;
+				}
 				case EStreamType::Subtitle:
+				{
 					BufferSourceInfo = ActivePeriods[i].BufferSourceInfoTxt;
 					break;
+				}
 				default:
+				{
 					break;
+				}
 			}
 		}
 	}
diff --git a/Engine/Plugins/Media/ElectraPlayer/Source/ElectraPlayerRuntime/Private/Runtime/Player/AdaptiveStreamingPlayer.h b/Engine/Plugins/Media/ElectraPlayer/Source/ElectraPlayerRuntime/Private/Runtime/Player/AdaptiveStreamingPlayer.h
--- a/Engine/Plugins/Media/ElectraPlayer/Source/ElectraPlayerRuntime/Private/Runtime/Player/AdaptiveStreamingPlayer.h
+++ b/Engine/Plugins/Media/ElectraPlayer/Source/ElectraPlayerRuntime/Private/Runtime/Player/AdaptiveStreamingPlayer.h
@@ -96,7 +96,7 @@ public:
 	 * Any later call to SelectTrackByMetadata() internally overwrites these initial attributes
 	 * with those of the explicitly selected track.
 	 */
-	virtual void SetInitialStreamAttributes(EStreamType StreamType, const FStreamSelectionAttributes& InitialSelection) = 0;
+	virtual void SetInitialStreamAttributes(EStreamType StreamType, FStreamSelectionAttributes InitialSelection) = 0;
 
 	/**
 	 * Enables or disables frame accurate seeking.
diff --git a/Engine/Plugins/Media/ElectraPlayer/Source/ElectraPlayerRuntime/Private/Runtime/Player/AdaptiveStreamingPlayerInternal.h b/Engine/Plugins/Media/ElectraPlayer/Source/ElectraPlayerRuntime/Private/Runtime/Player/AdaptiveStreamingPlayerInternal.h
--- a/Engine/Plugins/Media/ElectraPlayer/Source/ElectraPlayerRuntime/Private/Runtime/Player/AdaptiveStreamingPlayerInternal.h
+++ b/Engine/Plugins/Media/ElectraPlayer/Source/ElectraPlayerRuntime/Private/Runtime/Player/AdaptiveStreamingPlayerInternal.h
@@ -1126,7 +1126,7 @@ public:
 	void ModifyOptions(const FParamDict& InOptionsToSetOrChange, const FParamDict& InOptionsToClear) override;
 	FVariantValue GetMediaInfo(FName InKey) const override;
 
-	void SetInitialStreamAttributes(EStreamType StreamType, const FStreamSelectionAttributes& InitialSelection) override;
+	void SetInitialStreamAttributes(EStreamType StreamType, FStreamSelectionAttributes InitialSelection) override;
 	void EnableFrameAccurateSeeking(bool bEnabled) override;
 	void LoadBlob(TSharedPtr<FHTTPResourceRequest, ESPMode::ThreadSafe> InBlobLoadRequest) override;
 	void LoadManifest(const FString& manifestURL) override;
@@ -1565,7 +1565,6 @@ private:
 				Close,
 				ChangeBitrate,
 				LimitResolution,
-				InitialStreamAttributes,
 				SelectTrackByMetadata,
 				SelectTrackByAttributes,
 				DeselectTrack,
@@ -1749,14 +1748,6 @@ private:
 			Msg.Data.Emplace<FMessage::FResolution>(MoveTemp(reso));
 			TriggerSharedWorkerThread(MoveTemp(Msg));
 		}
-		void SendInitialStreamAttributeMessage(EStreamType InStreamType, const FStreamSelectionAttributes& InInitialSelection)
-		{
-			FMessage::FInitialStreamSelect iss { InStreamType, InInitialSelection };
-			FMessage Msg;
-			Msg.Type = FMessage::EType::InitialStreamAttributes;
-			Msg.Data.Emplace<FMessage::FInitialStreamSelect>(MoveTemp(iss));
-			TriggerSharedWorkerThread(MoveTemp(Msg));
-		}
 		void SendTrackSelectByMetadataMessage(EStreamType InStreamType, const FTrackMetadata& InTrackMetadata)
 		{
 			FMessage::FMetadataTrackSelection mts { InStreamType, InTrackMetadata };
diff --git a/Engine/Plugins/Media/ElectraPlayer/Source/ElectraPlayerRuntime/Private/Runtime/Player/DASH/ManifestBuilderDASH.cpp b/Engine/Plugins/Media/ElectraPlayer/Source/ElectraPlayerRuntime/Private/Runtime/Player/DASH/ManifestBuilderDASH.cpp
--- a/Engine/Plugins/Media/ElectraPlayer/Source/ElectraPlayerRuntime/Private/Runtime/Player/DASH/ManifestBuilderDASH.cpp
+++ b/Engine/Plugins/Media/ElectraPlayer/Source/ElectraPlayerRuntime/Private/Runtime/Player/DASH/ManifestBuilderDASH.cpp
@@ -1651,7 +1651,7 @@ void FManifestDASHInternal::PreparePeriodAdaptationSets(TSharedPtrTS<FPeriod> Pe
 			AdaptationSet->UniqueSequentialSetIndex = nAdapt;
 			AdaptationSet->AdaptationSet = MPDAdaptationSet;
 			AdaptationSet->PAR = MPDAdaptationSet->GetPAR();
-			AdaptationSet->Language = ISO639::RFC5646To639_1(MPDAdaptationSet->GetLanguage());
+			BCP47::ParseRFC5646Tag(AdaptationSet->LanguageTag, MPDAdaptationSet->GetLanguage());
 			AdaptationSet->SelectionPriority = (int32) MPDAdaptationSet->GetSelectionPriority();
 			// Set up labels. For now we ignore GroupLabel elements.
 			for(auto& itLabel : MPDAdaptationSet->GetLabels())
@@ -1913,7 +1913,7 @@ void FManifestDASHInternal::PreparePeriodAdaptationSets(TSharedPtrTS<FPeriod> Pe
 				Representation->CodecInfo.SetBitrate(Representation->Bitrate);
 
 				// Propagate the language code from the AdaptationSet into the codec info
-				Representation->CodecInfo.SetStreamLanguageCode(AdaptationSet->Language);
+				Representation->CodecInfo.SetStreamLanguageTag(AdaptationSet->LanguageTag);
 
 				// Set up codec details based on available attributes of the Representation or its enclosing AdaptationSet.
 				if (Representation->CodecInfo.IsVideoCodec())
@@ -2353,7 +2353,7 @@ void FManifestDASHInternal::PreparePeriodAdaptationSets(TSharedPtrTS<FPeriod> Pe
 						SwitchSet->Roles = SwitchedAS->Roles;
 						SwitchSet->Accessibilities = SwitchedAS->Accessibilities;
 						SwitchSet->PAR = SwitchedAS->PAR;
-						SwitchSet->Language = SwitchedAS->Language;
+						SwitchSet->LanguageTag = SwitchedAS->LanguageTag;
 					}
 					if (SwitchedAS->MaxBandwidth > SwitchSet->MaxBandwidth)
 					{
diff --git a/Engine/Plugins/Media/ElectraPlayer/Source/ElectraPlayerRuntime/Private/Runtime/Player/DASH/ManifestDASH.cpp b/Engine/Plugins/Media/ElectraPlayer/Source/ElectraPlayerRuntime/Private/Runtime/Player/DASH/ManifestDASH.cpp
--- a/Engine/Plugins/Media/ElectraPlayer/Source/ElectraPlayerRuntime/Private/Runtime/Player/DASH/ManifestDASH.cpp
+++ b/Engine/Plugins/Media/ElectraPlayer/Source/ElectraPlayerRuntime/Private/Runtime/Player/DASH/ManifestDASH.cpp
@@ -739,7 +739,11 @@ TSharedPtrTS<FManifestDASHInternal> FDASHPlayPeriod::GetCurrentManifest() const
 
 void FDASHPlayPeriod::SetStreamPreferences(EStreamType ForStreamType, const FStreamSelectionAttributes& StreamAttributes)
 {
-	if (ForStreamType == EStreamType::Audio)
+	if (ForStreamType == EStreamType::Video)
+	{
+		VideoStreamPreferences = StreamAttributes;
+	}
+	else if (ForStreamType == EStreamType::Audio)
 	{
 		AudioStreamPreferences = StreamAttributes;
 	}
@@ -810,38 +814,24 @@ TSharedPtrTS<FManifestDASHInternal::FAdaptationSet> FDASHPlayPeriod::SelectAdapt
 	if (NumAdaptationSets > 0)
 	{
 		int32 SelectedTypeIndex = 0;
-		// For video just select the first one for now.
-		if (StreamType == EStreamType::Video)
-		{
-			// Create a list of candidate AdaptationSets
-			for(int32 i=0; i<NumAdaptationSets; ++i)
-			{
-				TSharedPtrTS<FManifestDASHInternal::FAdaptationSet> V = StaticCastSharedPtr<FManifestDASHInternal::FAdaptationSet>(Period->GetAdaptationSetByTypeAndIndex(StreamType, i));
-				AddAdaptationSetToSelection(Selection, V, i);
-			}
-			// Prioritize the candidates based on user configuration.
-			PrioritizeSelection(Selection, StreamType, true, false);
-			if (Selection.Num())
-			{
-				// Take the highest prioritized set.
-				SelectedTypeIndex = Selection[0].Index;
-				AS = StaticCastSharedPtr<FManifestDASHInternal::FAdaptationSet>(Period->GetAdaptationSetByTypeAndIndex(StreamType, SelectedTypeIndex));
-			}
-		}
-		else if (StreamType == EStreamType::Audio || StreamType == EStreamType::Subtitle)
+		if (StreamType == EStreamType::Video || StreamType == EStreamType::Audio || StreamType == EStreamType::Subtitle)
 		{
 			// Check for a matching language.
 			// For now we ignore the track kind.
-			if (Attributes.Language_ISO639.IsSet())
+			if (Attributes.Language_RFC4647.IsSet())
 			{
-				// Create a list of candidate AdaptationSets
+				TArray<int32> CandidateIndices;
+				TArray<BCP47::FLanguageTag> CandList;
 				for(int32 i=0; i<NumAdaptationSets; ++i)
 				{
 					TSharedPtrTS<FManifestDASHInternal::FAdaptationSet> A = StaticCastSharedPtr<FManifestDASHInternal::FAdaptationSet>(Period->GetAdaptationSetByTypeAndIndex(StreamType, i));
-					if (A->GetLanguage().Equals(Attributes.Language_ISO639.GetValue()))
-					{
-						AddAdaptationSetToSelection(Selection, A, i);
-					}
+					CandList.Emplace(A->GetLanguageTag());
+				}
+				CandidateIndices = BCP47::FindExtendedFilteringMatch(CandList, Attributes.Language_RFC4647.GetValue());
+				for(int32 i=0; i<CandidateIndices.Num(); ++i)
+				{
+					TSharedPtrTS<FManifestDASHInternal::FAdaptationSet> A = StaticCastSharedPtr<FManifestDASHInternal::FAdaptationSet>(Period->GetAdaptationSetByTypeAndIndex(StreamType, CandidateIndices[i]));
+					AddAdaptationSetToSelection(Selection, A, CandidateIndices[i]);
 				}
 				// Prioritize the candidates based on user configuration.
 				PrioritizeSelection(Selection, StreamType, true, false);
@@ -862,7 +852,7 @@ TSharedPtrTS<FManifestDASHInternal::FAdaptationSet> FDASHPlayPeriod::SelectAdapt
 							{
 								if (Selection[i].CodecInfo.GetCodecName().Equals(PreferredCodec))
 								{
-									SelectedTypeIndex = i;
+									SelectedTypeIndex = Selection[i].Index;
 									break;
 								}
 							}
@@ -908,7 +898,7 @@ TSharedPtrTS<FManifestDASHInternal::FAdaptationSet> FDASHPlayPeriod::SelectAdapt
 			AS->GetMetaData(tm, StreamType);
 			OutBufferSourceInfo = MakeSharedTS<FBufferSourceInfo>();
 			OutBufferSourceInfo->Kind = tm.Kind;
-			OutBufferSourceInfo->Language = tm.Language;
+			OutBufferSourceInfo->LanguageTag = tm.LanguageTagRFC5646;
 			OutBufferSourceInfo->Codec = tm.HighestBandwidthCodec.GetCodecName();
 			OutBufferSourceInfo->HardIndex = SelectedTypeIndex;
 			OutBufferSourceInfo->PeriodID = Period->GetUniqueIdentifier();
@@ -1164,6 +1154,10 @@ IManifest::IPlayPeriod::ETrackChangeResult FDASHPlayPeriod::ChangeTrackStreamPre
 	// We cannot check if the the new stream to be selected is already the one that is active because a track change is
 	// triggered at the current playback position, while we could already be in a later period!
 	// Checking against a later period makes no sense, so we have to forcibly start over.
+	if (ForStreamType == EStreamType::Video)
+	{
+		return IManifest::IPlayPeriod::ETrackChangeResult::StartOver;
+	}
 	return IManifest::IPlayPeriod::ETrackChangeResult::NewPeriodNeeded;
 }
 
@@ -2948,6 +2942,12 @@ FManifestDASHInternal::FRepresentation::ESearchResult FManifestDASHInternal::FRe
 			return FManifestDASHInternal::FRepresentation::ESearchResult::BadType;
 		}
 		// We require segments to begin with SAP type 1 or 2 (preferably 1)
+		/*
+		else if (SegmentInfo.SAPType == 0)
+		{
+			LogMessage(InPlayerSessionServices, IInfoLog::ELevel::Warning, FString::Printf(TEXT("Segment index for Representation \"%s\" has unknown SAP_type (0). If it is not 1 or 2 playback may not work correctly"), *MPDRepresentation->GetID()));
+		}
+		*/
 		else if (SegmentInfo.SAPType != 1 && SegmentInfo.SAPType != 2)
 		{
 			PostError(InPlayerSessionServices, FString::Printf(TEXT("Segment index for Representation \"%s\" must have SAP_type 1 or 2 only!"), *MPDRepresentation->GetID()), ERRCODE_DASH_MPD_BAD_REPRESENTATION);
diff --git a/Engine/Plugins/Media/ElectraPlayer/Source/ElectraPlayerRuntime/Private/Runtime/Player/DASH/PlaylistReaderDASH_Internal.h b/Engine/Plugins/Media/ElectraPlayer/Source/ElectraPlayerRuntime/Private/Runtime/Player/DASH/PlaylistReaderDASH_Internal.h
--- a/Engine/Plugins/Media/ElectraPlayer/Source/ElectraPlayerRuntime/Private/Runtime/Player/DASH/PlaylistReaderDASH_Internal.h
+++ b/Engine/Plugins/Media/ElectraPlayer/Source/ElectraPlayerRuntime/Private/Runtime/Player/DASH/PlaylistReaderDASH_Internal.h
@@ -293,7 +293,6 @@ public:
 		const TArray<FString>& GetRoles() const										{ return Roles; }
 		const TArray<FString>& GetAccessibilities() const							{ return Accessibilities; }
 		const FTimeFraction& GetPAR() const											{ return PAR; }
-		//const FString& GetLanguage() const											{ return Language; }
 		int32 GetMaxBandwidth() const												{ return MaxBandwidth; }
 		int32 GetSelectionPriority() const											{ return SelectionPriority; }
 		bool GetIsUsable() const													{ return bIsUsable; }
@@ -417,7 +416,7 @@ public:
 					// TODO
 					check(!"need to add ID and language handling");
 					InOutMetadata.ID = TEXT("CC1");
-					InOutMetadata.Language = TEXT("en");
+					BCP47::ParseRFC5646Tag(InOutMetadata.LanguageTagRFC5646, FString(TEXT("en")));
 				}
 			}
 		}
@@ -426,7 +425,7 @@ public:
 		void GetMetaData(FTrackMetadata& OutMetadata, EStreamType StreamType) const
 		{
 			OutMetadata.ID = GetUniqueIdentifier();
-			OutMetadata.Language = GetLanguage();
+			OutMetadata.LanguageTagRFC5646 = GetLanguageTag();
 			OutMetadata.HighestBandwidth = GetMaxBandwidth();
 			OutMetadata.HighestBandwidthCodec = GetCodec();
 			if (Labels.Num())
@@ -488,9 +487,9 @@ public:
 		{
 			return Codec.GetCodecSpecifierRFC6381();
 		}
-		FString GetLanguage() const override
+		const BCP47::FLanguageTag& GetLanguageTag() const override
 		{
-			return Language;
+			return LanguageTag;
 		}
 		int32 GetNumberOfRepresentations() const override
 		{
@@ -521,7 +520,7 @@ public:
 		TArray<FString> Accessibilities;
 		TArray<FLabelElement> Labels;
 		FTimeFraction PAR;
-		FString Language;
+		BCP47::FLanguageTag LanguageTag;
 		int32 MaxBandwidth = 0;
 		int32 SelectionPriority = 1;
 		int32 UniqueSequentialSetIndex = 0;
diff --git a/Engine/Plugins/Media/ElectraPlayer/Source/ElectraPlayerRuntime/Private/Runtime/Player/HLS/PlaylistHLS.cpp b/Engine/Plugins/Media/ElectraPlayer/Source/ElectraPlayerRuntime/Private/Runtime/Player/HLS/PlaylistHLS.cpp
--- a/Engine/Plugins/Media/ElectraPlayer/Source/ElectraPlayerRuntime/Private/Runtime/Player/HLS/PlaylistHLS.cpp
+++ b/Engine/Plugins/Media/ElectraPlayer/Source/ElectraPlayerRuntime/Private/Runtime/Player/HLS/PlaylistHLS.cpp
@@ -601,6 +601,7 @@ FErrorDetail FActiveHLSPlaylist::CreateTrackMetadata(IPlayerSessionServices* InP
 				FMultiVariantPlaylistHLS::FInternalTrackMetadata& tm = InPathway->VideoTracks.Emplace_GetRef();
 				tm.Meta.ID = FString::Printf(TEXT("vid:%s"), *rg->Renditions[i].Name);
 				tm.Meta.Label = rg->Renditions[i].Name;	// Set the label to be the rendition's name. This gets used further down in comparisons!
+				tm.Meta.LanguageTagRFC5646 = rg->Renditions[i].LanguageRFC5646;
 				tm.Meta.Kind = i==0 ? TEXT("main") : TEXT("alternative");
 			}
 			for(int32 vtIdx=0; vtIdx<InPathway->VideoTracks.Num(); ++vtIdx)
@@ -673,7 +674,7 @@ FErrorDetail FActiveHLSPlaylist::CreateTrackMetadata(IPlayerSessionServices* InP
 				tm.Rendition = rg->Renditions[i];
 				tm.Meta.ID = FString::Printf(TEXT("aud:%s:%s"), *AudioGroups[0], *rg->Renditions[i].Name);
 				tm.Meta.Label = rg->Renditions[i].Name;
-				tm.Meta.Language = rg->Renditions[i].Language.IsEmpty() ? TEXT("und") : ISO639::RFC5646To639_1(rg->Renditions[i].Language);
+				tm.Meta.LanguageTagRFC5646 = rg->Renditions[i].LanguageRFC5646;
 				tm.Meta.Kind = i==0 ? TEXT("main") : TEXT("translation");
 					/*
 						An AUDIO Rendition MAY include the following characteristic:
@@ -703,7 +704,7 @@ FErrorDetail FActiveHLSPlaylist::CreateTrackMetadata(IPlayerSessionServices* InP
 				FMultiVariantPlaylistHLS::FInternalTrackMetadata& tm = InPathway->AudioTracks.Emplace_GetRef();
 				tm.Meta.ID = TEXT("aud:");
 				tm.Meta.Kind = TEXT("main");
-				tm.Meta.Language = TEXT("und");
+				//BCP47::ParseRFC5646Tag(tm.Meta.LanguageTagRFC5646, FString(TEXT("und")));
 				tm.bIsVariant = true;
 				AssignTracks(EStreamType::Audio, tm);
 			}
@@ -731,7 +732,7 @@ FErrorDetail FActiveHLSPlaylist::CreateTrackMetadata(IPlayerSessionServices* InP
 					tm.bIsVariant = true;
 					tm.Meta.ID = TEXT("aud:");
 					tm.Meta.Kind = TEXT("main");
-					tm.Meta.Language = TEXT("und");
+					//BCP47::ParseRFC5646Tag(tm.Meta.LanguageTagRFC5646, FString(TEXT("und")));
 					if (va[0].VariantGroupIndex < 0)
 					{
 						// This is the case where a variant stream includes an audio codec.
@@ -797,7 +798,7 @@ FErrorDetail FActiveHLSPlaylist::CreateTrackMetadata(IPlayerSessionServices* InP
 				tm.Rendition = rg->Renditions[i];
 				tm.Meta.ID = rg->Renditions[i].StableRenditionId;
 				tm.Meta.Label = rg->Renditions[i].Name;
-				tm.Meta.Language = rg->Renditions[i].Language.IsEmpty() ? TEXT("und") : ISO639::RFC5646To639_1(rg->Renditions[i].Language);
+				tm.Meta.LanguageTagRFC5646 = rg->Renditions[i].LanguageRFC5646;
 				tm.Meta.Kind = TEXT("subtitles");
 					/*
 						A SUBTITLES Rendition MAY include the following characteristics:
@@ -860,7 +861,7 @@ FErrorDetail FActiveHLSPlaylist::CreateTrackMetadata(IPlayerSessionServices* InP
 			TSharedPtrTS<FMultiVariantPlaylistHLS::FPlaybackAssetAdaptationSet> as = MakeSharedTS<FMultiVariantPlaylistHLS::FPlaybackAssetAdaptationSet>();
 			as->ID = tm.Meta.ID;
 			as->ListOfCodecs = tm.Meta.HighestBandwidthCodec.GetCodecSpecifierRFC6381();
-			as->Language = tm.Meta.Language;
+			as->LanguageTag = tm.Meta.LanguageTagRFC5646;
 			int32 sIdx = 0;
 			for(auto& sm : tm.Meta.StreamDetails)
 			{
@@ -2504,15 +2505,14 @@ void FActiveHLSPlaylist::FPlayPeriod::PrepareForPlay()
 			{
 				TArray<int32> CandidateIndices;
 				// Choose language?
-				if (Sel.Language_ISO639.IsSet())
+				if (Sel.Language_RFC4647.IsSet())
 				{
+					TArray<BCP47::FLanguageTag> CandList;
 					for(int32 i=0; i<Tracks->Num(); ++i)
 					{
-						if ((*Tracks)[i].Meta.Language.Equals(Sel.Language_ISO639.GetValue()))
-						{
-							CandidateIndices.Emplace(i);
-						}
+						CandList.Emplace((*Tracks)[i].Meta.LanguageTagRFC5646);
 					}
+					CandidateIndices = BCP47::FindExtendedFilteringMatch(CandList, Sel.Language_RFC4647.GetValue());
 				}
 				// If there are multiple language candidates narrow the list down by kind.
 				if (CandidateIndices.Num() && Sel.Kind.IsSet())
@@ -2522,7 +2522,7 @@ void FActiveHLSPlaylist::FPlayPeriod::PrepareForPlay()
 					{
 						if ((*Tracks)[CandidateIndices[i]].Meta.Kind.Equals(Sel.Kind.GetValue()))
 						{
-							TempList.Emplace(i);
+							TempList.Emplace(CandidateIndices[i]);
 						}
 					}
 					// If there are new candidates update the list. If everything is filtered out, keep the previous list.
@@ -2550,9 +2550,9 @@ void FActiveHLSPlaylist::FPlayPeriod::PrepareForPlay()
 			st.BufferSourceInfo->PeriodID = TimelineMediaAsset->GetAssetIdentifier();
 			st.BufferSourceInfo->PeriodAdaptationSetID = st.MetaID;
 			st.BufferSourceInfo->Kind = Track.Meta.Kind;
-			st.BufferSourceInfo->Language = Track.Meta.Language;
+			st.BufferSourceInfo->LanguageTag = Track.Meta.LanguageTagRFC5646;
 			st.BufferSourceInfo->Codec = Util::GetBaseCodec(Track.Meta.HighestBandwidthCodec.GetCodecSpecifierRFC6381());
-			st.BufferSourceInfo->HardIndex = st.StreamIndex;
+			st.BufferSourceInfo->HardIndex = st.TrackIndex;
 
 			FStreamLoadRequest LoadReq;
 			if (TimelineMediaAsset->GetVariantPlaylist(LoadReq.Request, PlayerSessionServices, StreamType, Pwy, SelectedTrackIndex,0, SelectedTrackStream[0].TrackIndex,0).IsOK())
diff --git a/Engine/Plugins/Media/ElectraPlayer/Source/ElectraPlayerRuntime/Private/Runtime/Player/HLS/PlaylistHLS.h b/Engine/Plugins/Media/ElectraPlayer/Source/ElectraPlayerRuntime/Private/Runtime/Player/HLS/PlaylistHLS.h
--- a/Engine/Plugins/Media/ElectraPlayer/Source/ElectraPlayerRuntime/Private/Runtime/Player/HLS/PlaylistHLS.h
+++ b/Engine/Plugins/Media/ElectraPlayer/Source/ElectraPlayerRuntime/Private/Runtime/Player/HLS/PlaylistHLS.h
@@ -6,6 +6,7 @@
 #include "Containers/Array.h"
 #include "Containers/LruCache.h"
 #include "Utilities/URLParser.h"
+#include "Utilities/BCP47-Helpers.h"
 #include "Player/PlayerSessionServices.h"
 #include "Player/AdaptiveStreamingPlayerResourceRequest.h"
 #include "Player/HLS/PlaylistParserHLS.h"
@@ -84,8 +85,8 @@ struct FMultiVariantPlaylistHLS
 		{ return ID; }
 		FString GetListOfCodecs() const override
 		{ return ListOfCodecs; }
-		FString GetLanguage() const override
-		{ return Language; }
+		const BCP47::FLanguageTag& GetLanguageTag() const override
+		{ return LanguageTag; }
 		int32 GetNumberOfRepresentations() const override
 		{ return Representations.Num(); }
 		bool IsLowLatencyEnabled() const override
@@ -105,9 +106,9 @@ struct FMultiVariantPlaylistHLS
 		}
 
 		TArray<TSharedPtrTS<FPlaybackAssetRepresentation>> Representations;
+		BCP47::FLanguageTag LanguageTag;
 		FString ID;
 		FString ListOfCodecs;
-		FString Language;
 	};
 
 
@@ -117,19 +118,19 @@ struct FMultiVariantPlaylistHLS
 		FStreamCodecInformation ParsedCodecFromStreamInf;
 		FString CodecNameFromStreamInf;
 
-		FString Name;				// NAME
-		FString Language;			// LANGUAGE
-		FString AssocLanguage;		// ASSOC-LANGUAGE
-		FString StableRenditionId;	// STABLE-RENDITION-ID
-		FString URI;				// URI
-		FString InstreamId;			// INSTREAM-ID
-		FString Characteristics;	// CHARACTERISTICS
-		int32 SampleRate = 0;		// SAMPLE-RATE
-		int32 Channels = 0;			// CHANNELS (first channel count part only)
-		int32 BitDepth = 0;			// BIT-DEPTH
-		bool bDefault = false;		// DEFAULT
-		bool bAutoSelect = false;	// AUTOSELECT
-		bool bForced = false;		// FORCED
+		FString Name;								// NAME
+		BCP47::FLanguageTag LanguageRFC5646;		// LANGUAGE
+		BCP47::FLanguageTag AssocLanguageRFC5646;	// ASSOC-LANGUAGE
+		FString StableRenditionId;					// STABLE-RENDITION-ID
+		FString URI;								// URI
+		FString InstreamId;							// INSTREAM-ID
+		FString Characteristics;					// CHARACTERISTICS
+		int32 SampleRate = 0;						// SAMPLE-RATE
+		int32 Channels = 0;							// CHANNELS (first channel count part only)
+		int32 BitDepth = 0;							// BIT-DEPTH
+		bool bDefault = false;						// DEFAULT
+		bool bAutoSelect = false;					// AUTOSELECT
+		bool bForced = false;						// FORCED
 	};
 
 	struct FRenditionGroup
diff --git a/Engine/Plugins/Media/ElectraPlayer/Source/ElectraPlayerRuntime/Private/Runtime/Player/HLS/PlaylistHandlerHLS.cpp b/Engine/Plugins/Media/ElectraPlayer/Source/ElectraPlayerRuntime/Private/Runtime/Player/HLS/PlaylistHandlerHLS.cpp
--- a/Engine/Plugins/Media/ElectraPlayer/Source/ElectraPlayerRuntime/Private/Runtime/Player/HLS/PlaylistHandlerHLS.cpp
+++ b/Engine/Plugins/Media/ElectraPlayer/Source/ElectraPlayerRuntime/Private/Runtime/Player/HLS/PlaylistHandlerHLS.cpp
@@ -1696,20 +1696,30 @@ bool FPlaylistHandlerHLS::BuildMultiVariantPlaylist(TSharedPtr<FMultiVariantPlay
 							// LANGUAGE ?
 							else if (AttrName.Equals(TEXT("LANGUAGE")))
 							{
-								if (!Attributes[i].GetValue(NewRendition.Language, NewVariantPlaylist->VariableSubstitutions))
+								FString Lang5646;
+								if (!Attributes[i].GetValue(Lang5646, NewVariantPlaylist->VariableSubstitutions))
 								{
 									LastError = PostError(FString::Printf(TEXT("EXT-X-MEDIA has a bad LANGUAGE value")), HLS::ERRCODE_PLAYLIST_PARSING_FAILED);
 									return false;
 								}
+								if (!BCP47::ParseRFC5646Tag(NewRendition.LanguageRFC5646, Lang5646))
+								{
+									LogMessage(IInfoLog::ELevel::Warning, FString::Printf(TEXT("LANGUAGE \"%s\" is not a valid BCP-47 tag!"), *Lang5646));
+								}
 							}
 							// ASSOC-LANGUAGE ?
 							else if (AttrName.Equals(TEXT("ASSOC-LANGUAGE")))
 							{
-								if (!Attributes[i].GetValue(NewRendition.AssocLanguage, NewVariantPlaylist->VariableSubstitutions))
+								FString Lang5646;
+								if (!Attributes[i].GetValue(Lang5646, NewVariantPlaylist->VariableSubstitutions))
 								{
 									LastError = PostError(FString::Printf(TEXT("EXT-X-MEDIA has a bad ASSOC-LANGUAGE value")), HLS::ERRCODE_PLAYLIST_PARSING_FAILED);
 									return false;
 								}
+								if (!BCP47::ParseRFC5646Tag(NewRendition.AssocLanguageRFC5646, Lang5646))
+								{
+									LogMessage(IInfoLog::ELevel::Warning, FString::Printf(TEXT("ASSOC-LANGUAGE \"%s\" is not a valid BCP-47 tag!"), *Lang5646));
+								}
 							}
 							// NAME ?
 							else if (AttrName.Equals(TEXT("NAME")))
diff --git a/Engine/Plugins/Media/ElectraPlayer/Source/ElectraPlayerRuntime/Private/Runtime/Player/PlaybackTimeline.h b/Engine/Plugins/Media/ElectraPlayer/Source/ElectraPlayerRuntime/Private/Runtime/Player/PlaybackTimeline.h
--- a/Engine/Plugins/Media/ElectraPlayer/Source/ElectraPlayerRuntime/Private/Runtime/Player/PlaybackTimeline.h
+++ b/Engine/Plugins/Media/ElectraPlayer/Source/ElectraPlayerRuntime/Private/Runtime/Player/PlaybackTimeline.h
@@ -40,7 +40,7 @@ namespace Electra
 
 		virtual FString GetListOfCodecs() const = 0;
 
-		virtual FString GetLanguage() const = 0;
+		virtual const BCP47::FLanguageTag& GetLanguageTag() const = 0;
 
 		virtual int32 GetNumberOfRepresentations() const = 0;
 
diff --git a/Engine/Plugins/Media/ElectraPlayer/Source/ElectraPlayerRuntime/Private/Runtime/Player/mkv/ManifestMKV.cpp b/Engine/Plugins/Media/ElectraPlayer/Source/ElectraPlayerRuntime/Private/Runtime/Player/mkv/ManifestMKV.cpp
--- a/Engine/Plugins/Media/ElectraPlayer/Source/ElectraPlayerRuntime/Private/Runtime/Player/mkv/ManifestMKV.cpp
+++ b/Engine/Plugins/Media/ElectraPlayer/Source/ElectraPlayerRuntime/Private/Runtime/Player/mkv/ManifestMKV.cpp
@@ -432,17 +432,18 @@ TSharedPtrTS<FTrackMetadata> FManifestMKVInternal::FPlayPeriodMKV::SelectMetadat
 		}
 		if (Metadata.Num())
 		{
-			// We do not look at the 'kind' here, only the language.
+			// We do not look at the 'kind' or 'codec' here, only the language.
 			// Set the first track as default in case we do not find the one we're looking for.
-			if (InAttributes.Language_ISO639.IsSet())
+			if (InAttributes.Language_RFC4647.IsSet())
 			{
+				TArray<int32> CandidateIndices;
+				TArray<BCP47::FLanguageTag> CandList;
 				for(auto &Meta : Metadata)
 				{
-					if (Meta.Language.Equals(InAttributes.Language_ISO639.GetValue()))
-					{
-						return MakeSharedTS<FTrackMetadata>(Meta);
-					}
+					CandList.Emplace(Meta.LanguageTagRFC5646);
 				}
+				CandidateIndices = BCP47::FindExtendedFilteringMatch(CandList, InAttributes.Language_RFC4647.GetValue());
+				return MakeSharedTS<FTrackMetadata>(CandidateIndices.Num() ? Metadata[CandidateIndices[0]] : Metadata[0]);
 			}
 			return MakeSharedTS<FTrackMetadata>(Metadata[0]);
 		}
@@ -456,7 +457,7 @@ void FManifestMKVInternal::FPlayPeriodMKV::MakeBufferSourceInfoFromMetadata(EStr
 	{
 		OutBufferSourceInfo = MakeSharedTS<FBufferSourceInfo>();
 		OutBufferSourceInfo->Kind = InMetadata->Kind;
-		OutBufferSourceInfo->Language = InMetadata->Language;
+		OutBufferSourceInfo->LanguageTag = InMetadata->LanguageTagRFC5646;
 		OutBufferSourceInfo->Codec = InMetadata->HighestBandwidthCodec.GetCodecName();
 		TSharedPtrTS<FTimelineAssetMKV> Asset = MediaAsset.Pin();
 		OutBufferSourceInfo->PeriodID = Asset->GetUniqueIdentifier();
@@ -937,7 +938,7 @@ FErrorDetail FManifestMKVInternal::FAdaptationSetMKV::CreateFrom(const IParserMK
 	{
 		CodecRFC6381 = Representation->GetCodecInformation().GetCodecSpecifierRFC6381();
 		UniqueIdentifier = Representation->GetUniqueIdentifier();
-		Language = InTrack->GetLanguage();
+		LanguageTag = InTrack->GetLanguageTag();
 	}
 	return Error;
 }
diff --git a/Engine/Plugins/Media/ElectraPlayer/Source/ElectraPlayerRuntime/Private/Runtime/Player/mkv/ManifestMKV.h b/Engine/Plugins/Media/ElectraPlayer/Source/ElectraPlayerRuntime/Private/Runtime/Player/mkv/ManifestMKV.h
--- a/Engine/Plugins/Media/ElectraPlayer/Source/ElectraPlayerRuntime/Private/Runtime/Player/mkv/ManifestMKV.h
+++ b/Engine/Plugins/Media/ElectraPlayer/Source/ElectraPlayerRuntime/Private/Runtime/Player/mkv/ManifestMKV.h
@@ -115,8 +115,9 @@ public:
 		{ return UniqueIdentifier; }
 		FString GetListOfCodecs() const override
 		{ return CodecRFC6381; }
-		FString GetLanguage() const override
-		{ return Language; }
+		const BCP47::FLanguageTag& GetLanguageTag() const override
+		{ return LanguageTag; }
+
 		int32 GetNumberOfRepresentations() const override
 		{ return Representation.IsValid() ? 1 : 0; }
 		bool IsLowLatencyEnabled() const override
@@ -127,7 +128,7 @@ public:
 		{ return Representation.IsValid() && Representation->GetUniqueIdentifier() == InUniqueIdentifier ? Representation : TSharedPtrTS<IPlaybackAssetRepresentation>(); }
 	private:
 		TSharedPtrTS<FRepresentationMKV> Representation;
-		FString Language;
+		BCP47::FLanguageTag LanguageTag;
 		FString CodecRFC6381;
 		FString UniqueIdentifier;
 	};
@@ -216,7 +217,7 @@ public:
 				{
 					FTrackMetadata tm;
 					tm.ID = AdaptSet->GetUniqueIdentifier();
-					tm.Language = AdaptSet->GetLanguage();
+					tm.LanguageTagRFC5646 = AdaptSet->GetLanguageTag();
 					tm.Kind = i==0 ? TEXT("main") : TEXT("translation");
 
 					for(int32 j=0, jMax=AdaptSet->GetNumberOfRepresentations(); j<jMax; ++j)
diff --git a/Engine/Plugins/Media/ElectraPlayer/Source/ElectraPlayerRuntime/Private/Runtime/Player/mkv/StreamReaderMKV.cpp b/Engine/Plugins/Media/ElectraPlayer/Source/ElectraPlayerRuntime/Private/Runtime/Player/mkv/StreamReaderMKV.cpp
--- a/Engine/Plugins/Media/ElectraPlayer/Source/ElectraPlayerRuntime/Private/Runtime/Player/mkv/StreamReaderMKV.cpp
+++ b/Engine/Plugins/Media/ElectraPlayer/Source/ElectraPlayerRuntime/Private/Runtime/Player/mkv/StreamReaderMKV.cpp
@@ -274,7 +274,7 @@ void FStreamReaderMKV::HandleRequest()
 	{
 		EStreamType Type;
 		FString Kind;
-		FString Language;
+		BCP47::FLanguageTag LanguageTag;
 		FString PeriodID;
 		FString AdaptationSetID;
 		FString RepresentationID;
@@ -293,7 +293,7 @@ void FStreamReaderMKV::HandleRequest()
 			for(int32 nAdapt=0; nAdapt<NumAdapt; ++nAdapt)
 			{
 				TSharedPtrTS<IPlaybackAssetAdaptationSet> AdaptationSet = Request->MediaAsset->GetAdaptationSetByTypeAndIndex(TypesOfSupportedTracks[nStrType], nAdapt);
-				FString Language = AdaptationSet->GetLanguage();
+				const BCP47::FLanguageTag& LanguageTag = AdaptationSet->GetLanguageTag();
 				FString AdaptID = AdaptationSet->GetUniqueIdentifier();
 				int32 NumRepr = AdaptationSet->GetNumberOfRepresentations();
 				for(int32 nRepr=0; nRepr<NumRepr; ++nRepr)
@@ -305,7 +305,7 @@ void FStreamReaderMKV::HandleRequest()
 					LexFromString(TrackId, *ReprID);
 					FPlaylistTrackMetadata tmd;
 					tmd.Type = TypesOfSupportedTracks[nStrType];
-					tmd.Language = Language;
+					tmd.LanguageTag = LanguageTag;
 					tmd.Kind = nRepr == 0 ? TEXT("main") : TEXT("translation");	// perhaps this should come from querying the metadata instead of doing this over here...
 					tmd.PeriodID = Request->MediaAsset->GetUniqueIdentifier();
 					tmd.AdaptationSetID = AdaptID;
@@ -420,7 +420,7 @@ void FStreamReaderMKV::HandleRequest()
 					st.CSD->ParsedInfo.SetBitrate(st.Bitrate);
 				}
 				meta->Kind = SelectedTrackMetadata->Kind;
-				meta->Language = SelectedTrackMetadata->Language;
+				meta->LanguageTag = SelectedTrackMetadata->LanguageTag;
 				meta->Codec = st.CSD.IsValid() ? st.CSD->ParsedInfo.GetCodecName() : FString();
 				meta->PeriodID = SelectedTrackMetadata->PeriodID;
 				meta->PeriodAdaptationSetID = SelectedTrackMetadata->PeriodID + TEXT(".") + SelectedTrackMetadata->AdaptationSetID;
diff --git a/Engine/Plugins/Media/ElectraPlayer/Source/ElectraPlayerRuntime/Private/Runtime/Player/mp4/ManifestMP4.cpp b/Engine/Plugins/Media/ElectraPlayer/Source/ElectraPlayerRuntime/Private/Runtime/Player/mp4/ManifestMP4.cpp
--- a/Engine/Plugins/Media/ElectraPlayer/Source/ElectraPlayerRuntime/Private/Runtime/Player/mp4/ManifestMP4.cpp
+++ b/Engine/Plugins/Media/ElectraPlayer/Source/ElectraPlayerRuntime/Private/Runtime/Player/mp4/ManifestMP4.cpp
@@ -11,6 +11,7 @@
 #include "Utilities/TimeUtilities.h"
 #include "HAL/LowLevelMemTracker.h"
 #include "Utilities/UtilsMP4.h"
+#include "Utilities/BCP47-Helpers.h"
 #include "ElectraPlayerPrivate.h"
 #include "Player/AdaptiveStreamingPlayerResourceRequest.h"
 #include "Player/AdaptivePlayerOptionKeynames.h"
@@ -453,17 +454,18 @@ TSharedPtrTS<FTrackMetadata> FManifestMP4Internal::FPlayPeriodMP4::SelectMetadat
 		}
 		if (Metadata.Num())
 		{
-			// We do not look at the 'kind' here, only the language.
+			// We do not look at the 'kind' or 'codec' here, only the language.
 			// Set the first track as default in case we do not find the one we're looking for.
-			if (InAttributes.Language_ISO639.IsSet())
+			if (InAttributes.Language_RFC4647.IsSet())
 			{
+				TArray<int32> CandidateIndices;
+				TArray<BCP47::FLanguageTag> CandList;
 				for(auto &Meta : Metadata)
 				{
-					if (Meta.Language.Equals(InAttributes.Language_ISO639.GetValue()))
-					{
-						return MakeSharedTS<FTrackMetadata>(Meta);
-					}
+					CandList.Emplace(Meta.LanguageTagRFC5646);
 				}
+				CandidateIndices = BCP47::FindExtendedFilteringMatch(CandList, InAttributes.Language_RFC4647.GetValue());
+				return MakeSharedTS<FTrackMetadata>(CandidateIndices.Num() ? Metadata[CandidateIndices[0]] : Metadata[0]);
 			}
 			return MakeSharedTS<FTrackMetadata>(Metadata[0]);
 		}
@@ -477,7 +479,7 @@ void FManifestMP4Internal::FPlayPeriodMP4::MakeBufferSourceInfoFromMetadata(EStr
 	{
 		OutBufferSourceInfo = MakeSharedTS<FBufferSourceInfo>();
 		OutBufferSourceInfo->Kind = InMetadata->Kind;
-		OutBufferSourceInfo->Language = InMetadata->Language;
+		OutBufferSourceInfo->LanguageTag = InMetadata->LanguageTagRFC5646;
 		OutBufferSourceInfo->Codec = InMetadata->HighestBandwidthCodec.GetCodecName();
 		TSharedPtrTS<FTimelineAssetMP4> Asset = MediaAsset.Pin();
 		OutBufferSourceInfo->PeriodID = Asset->GetUniqueIdentifier();
@@ -743,16 +745,24 @@ FErrorDetail FManifestMP4Internal::FTimelineAssetMP4::Build(IPlayerSessionServic
 					switch(Track->GetCodecInformation().GetStreamType())
 					{
 						case EStreamType::Video:
+						{
 							VideoAdaptationSets.Add(AdaptationSet);
 							break;
+						}
 						case EStreamType::Audio:
+						{
 							AudioAdaptationSets.Add(AdaptationSet);
 							break;
+						}
 						case EStreamType::Subtitle:
+						{
 							SubtitleAdaptationSets.Add(AdaptationSet);
 							break;
+						}
 						default:
+						{
 							break;
+						}
 					}
 					FTimeValue TrkDur;
 					TrkDur.SetFromTimeFraction(Track->GetDuration());
@@ -1343,7 +1353,7 @@ FErrorDetail FManifestMP4Internal::FAdaptationSetMP4::CreateFrom(const IParserIS
 	{
 		CodecRFC6381	 = Representation->GetCodecInformation().GetCodecSpecifierRFC6381();
 		UniqueIdentifier = Representation->GetUniqueIdentifier();
-		Language		 = InTrack->GetLanguage();
+		LanguageTag		 = InTrack->GetLanguageTag();
 	}
 	return err;
 }
diff --git a/Engine/Plugins/Media/ElectraPlayer/Source/ElectraPlayerRuntime/Private/Runtime/Player/mp4/ManifestMP4.h b/Engine/Plugins/Media/ElectraPlayer/Source/ElectraPlayerRuntime/Private/Runtime/Player/mp4/ManifestMP4.h
--- a/Engine/Plugins/Media/ElectraPlayer/Source/ElectraPlayerRuntime/Private/Runtime/Player/mp4/ManifestMP4.h
+++ b/Engine/Plugins/Media/ElectraPlayer/Source/ElectraPlayerRuntime/Private/Runtime/Player/mp4/ManifestMP4.h
@@ -118,8 +118,8 @@ public:
 		{ return UniqueIdentifier; }
 		FString GetListOfCodecs() const override
 		{ return CodecRFC6381; }
-		FString GetLanguage() const override
-		{ return Language; }
+		const BCP47::FLanguageTag& GetLanguageTag() const override
+		{ return LanguageTag; }
 		int32 GetNumberOfRepresentations() const override
 		{ return Representation.IsValid() ? 1 : 0; }
 		bool IsLowLatencyEnabled() const override
@@ -130,7 +130,7 @@ public:
 		{ return Representation.IsValid() && Representation->GetUniqueIdentifier() == InUniqueIdentifier ? Representation : TSharedPtrTS<IPlaybackAssetRepresentation>(); }
 	private:
 		TSharedPtrTS<FRepresentationMP4>		Representation;
-		FString									Language;
+		BCP47::FLanguageTag						LanguageTag;
 		FString									CodecRFC6381;
 		FString									UniqueIdentifier;
 	};
@@ -201,7 +201,7 @@ public:
 				{
 					FTrackMetadata tm;
 					tm.ID = AdaptSet->GetUniqueIdentifier();
-					tm.Language = AdaptSet->GetLanguage();
+					tm.LanguageTagRFC5646 = AdaptSet->GetLanguageTag();
 					tm.Kind = i==0 ? TEXT("main") : TEXT("translation");
 
 					for(int32 j=0, jMax=AdaptSet->GetNumberOfRepresentations(); j<jMax; ++j)
diff --git a/Engine/Plugins/Media/ElectraPlayer/Source/ElectraPlayerRuntime/Private/Runtime/Player/mp4/StreamReaderMP4.cpp b/Engine/Plugins/Media/ElectraPlayer/Source/ElectraPlayerRuntime/Private/Runtime/Player/mp4/StreamReaderMP4.cpp
--- a/Engine/Plugins/Media/ElectraPlayer/Source/ElectraPlayerRuntime/Private/Runtime/Player/mp4/StreamReaderMP4.cpp
+++ b/Engine/Plugins/Media/ElectraPlayer/Source/ElectraPlayerRuntime/Private/Runtime/Player/mp4/StreamReaderMP4.cpp
@@ -279,7 +279,7 @@ void FStreamReaderMP4::HandleRequest()
 	{
 		EStreamType		Type;
 		FString			Kind;
-		FString			Language;
+		BCP47::FLanguageTag LanguageTag;
 		FString			PeriodID;
 		FString			AdaptationSetID;
 		FString			RepresentationID;
@@ -297,7 +297,7 @@ void FStreamReaderMP4::HandleRequest()
 			for(int32 nAdapt=0; nAdapt<NumAdapt; ++nAdapt)
 			{
 				TSharedPtrTS<IPlaybackAssetAdaptationSet> AdaptationSet = Request->MediaAsset->GetAdaptationSetByTypeAndIndex(TypesOfSupportedTracks[nStrType], nAdapt);
-				FString Language = AdaptationSet->GetLanguage();
+				const BCP47::FLanguageTag& LanguageTag = AdaptationSet->GetLanguageTag();
 				FString AdaptID = AdaptationSet->GetUniqueIdentifier();
 				int32 NumRepr = AdaptationSet->GetNumberOfRepresentations();
 				for(int32 nRepr=0; nRepr<NumRepr; ++nRepr)
@@ -310,7 +310,7 @@ void FStreamReaderMP4::HandleRequest()
 					LexFromString(TrackId, *ReprID);
 					FPlaylistTrackMetadata tmd;
 					tmd.Type			 = TypesOfSupportedTracks[nStrType];
-					tmd.Language		 = Language;
+					tmd.LanguageTag		 = LanguageTag;
 					tmd.Kind			 = nRepr == 0 ? TEXT("main") : TEXT("translation");	// perhaps this should come from querying the metadata instead of doing this over here...
 					tmd.PeriodID		 = Request->MediaAsset->GetUniqueIdentifier();
 					tmd.AdaptationSetID  = AdaptID;
@@ -448,7 +448,7 @@ void FStreamReaderMP4::HandleRequest()
 						st.CSD->ParsedInfo.SetBitrate(st.Bitrate);
 					}
 					meta->Kind = SelectedTrackMetadata->Kind;
-					meta->Language = SelectedTrackMetadata->Language;
+					meta->LanguageTag = SelectedTrackMetadata->LanguageTag;
 					meta->Codec = st.CSD.IsValid() ? st.CSD->ParsedInfo.GetCodecName() : FString();
 					meta->PeriodID = SelectedTrackMetadata->PeriodID;
 					meta->PeriodAdaptationSetID = SelectedTrackMetadata->PeriodID + TEXT(".") + SelectedTrackMetadata->AdaptationSetID;
diff --git a/Engine/Plugins/Media/ElectraPlayer/Source/ElectraPlayerRuntime/Private/Runtime/Player/mpegaudio/ManifestMPEGAudio.cpp b/Engine/Plugins/Media/ElectraPlayer/Source/ElectraPlayerRuntime/Private/Runtime/Player/mpegaudio/ManifestMPEGAudio.cpp
--- a/Engine/Plugins/Media/ElectraPlayer/Source/ElectraPlayerRuntime/Private/Runtime/Player/mpegaudio/ManifestMPEGAudio.cpp
+++ b/Engine/Plugins/Media/ElectraPlayer/Source/ElectraPlayerRuntime/Private/Runtime/Player/mpegaudio/ManifestMPEGAudio.cpp
@@ -483,17 +483,18 @@ TSharedPtrTS<FTrackMetadata> FManifestMPEGAudioInternal::FPlayPeriodMPEGAudio::S
 		}
 		if (Metadata.Num())
 		{
-			// We do not look at the 'kind' here, only the language.
+			// We do not look at the 'kind' or 'codec' here, only the language.
 			// Set the first track as default in case we do not find the one we're looking for.
-			if (InAttributes.Language_ISO639.IsSet())
+			if (InAttributes.Language_RFC4647.IsSet())
 			{
+				TArray<int32> CandidateIndices;
+				TArray<BCP47::FLanguageTag> CandList;
 				for(auto &Meta : Metadata)
 				{
-					if (Meta.Language.Equals(InAttributes.Language_ISO639.GetValue()))
-					{
-						return MakeSharedTS<FTrackMetadata>(Meta);
-					}
+					CandList.Emplace(Meta.LanguageTagRFC5646);
 				}
+				CandidateIndices = BCP47::FindExtendedFilteringMatch(CandList, InAttributes.Language_RFC4647.GetValue());
+				return MakeSharedTS<FTrackMetadata>(CandidateIndices.Num() ? Metadata[CandidateIndices[0]] : Metadata[0]);
 			}
 			return MakeSharedTS<FTrackMetadata>(Metadata[0]);
 		}
@@ -507,7 +508,7 @@ void FManifestMPEGAudioInternal::FPlayPeriodMPEGAudio::MakeBufferSourceInfoFromM
 	{
 		OutBufferSourceInfo = MakeSharedTS<FBufferSourceInfo>();
 		OutBufferSourceInfo->Kind = InMetadata->Kind;
-		OutBufferSourceInfo->Language = InMetadata->Language;
+		OutBufferSourceInfo->LanguageTag = InMetadata->LanguageTagRFC5646;
 		OutBufferSourceInfo->Codec = InMetadata->HighestBandwidthCodec.GetCodecName();
 		TSharedPtrTS<FTimelineAssetMPEGAudio> Asset = MediaAsset.Pin();
 		OutBufferSourceInfo->PeriodID = Asset->GetUniqueIdentifier();
@@ -1066,7 +1067,7 @@ FErrorDetail FManifestMPEGAudioInternal::FAdaptationSetMPEGAudio::CreateFrom(con
 	{
 		CodecRFC6381 = Representation->GetCodecInformation().GetCodecSpecifierRFC6381();
 		UniqueIdentifier = Representation->GetUniqueIdentifier();
-		Language = TEXT("und");
+		BCP47::ParseRFC5646Tag(LanguageTag, Representation->GetLanguage());
 		bIsLivecast = InHeader.bIsLive;
 	}
 	return err;
@@ -1100,6 +1101,18 @@ FErrorDetail FManifestMPEGAudioInternal::FRepresentationMPEGAudio::CreateFrom(co
 	{
 		Name = TEXT("Unknown");
 	}
+	
+	// Language?
+	if (InHeader.ID3v2.IsValid() && InHeader.ID3v2->HaveTag(Utils::Make4CC('T','L','A','N')))
+	{
+		MPEG::FID3V2Metadata::FItem TLan;
+		InHeader.ID3v2->GetTag(TLan, Utils::Make4CC('T','L','A','N'));
+		Language639_2 = TLan.Value.GetValue<FString>();
+	}
+	else
+	{
+		Language639_2 = TEXT("und");
+	}
 
 	Bitrate = InHeader.AverageBitrate ? InHeader.AverageBitrate : InHeader.Bitrate;
 	Bitrate = Bitrate <= 0 ? 64000 : Bitrate;
diff --git a/Engine/Plugins/Media/ElectraPlayer/Source/ElectraPlayerRuntime/Private/Runtime/Player/mpegaudio/ManifestMPEGAudio.h b/Engine/Plugins/Media/ElectraPlayer/Source/ElectraPlayerRuntime/Private/Runtime/Player/mpegaudio/ManifestMPEGAudio.h
--- a/Engine/Plugins/Media/ElectraPlayer/Source/ElectraPlayerRuntime/Private/Runtime/Player/mpegaudio/ManifestMPEGAudio.h
+++ b/Engine/Plugins/Media/ElectraPlayer/Source/ElectraPlayerRuntime/Private/Runtime/Player/mpegaudio/ManifestMPEGAudio.h
@@ -131,10 +131,15 @@ public:
 		{ return Name; }
 		void SetName(const FString& InNewName)
 		{ Name = InNewName; }
+		const FString& GetLanguage() const
+		{ return Language639_2; }
+		void SetLanguage(const FString& InNewName)
+		{ Language639_2 = InNewName; }
 	private:
 		FStreamCodecInformation CodecInformation;
 		FString UniqueIdentifier;
 		FString Name;
+		FString Language639_2;
 		int32 Bitrate = 0;
 	};
 
@@ -150,8 +155,8 @@ public:
 		{ return UniqueIdentifier; }
 		FString GetListOfCodecs() const override
 		{ return CodecRFC6381; }
-		FString GetLanguage() const override
-		{ return Language; }
+		const BCP47::FLanguageTag& GetLanguageTag() const override
+		{ return LanguageTag; }
 		int32 GetNumberOfRepresentations() const override
 		{ return Representation.IsValid() ? 1 : 0; }
 		bool IsLowLatencyEnabled() const override;
@@ -161,7 +166,7 @@ public:
 		{ return Representation.IsValid() && Representation->GetUniqueIdentifier() == InUniqueIdentifier ? Representation : TSharedPtrTS<IPlaybackAssetRepresentation>(); }
 	private:
 		TSharedPtrTS<FRepresentationMPEGAudio> Representation;
-		FString Language;
+		BCP47::FLanguageTag LanguageTag;
 		FString CodecRFC6381;
 		FString UniqueIdentifier;
 		bool bIsLivecast = false;
@@ -214,7 +219,7 @@ public:
 				{
 					FTrackMetadata tm;
 					tm.ID = AdaptSet->GetUniqueIdentifier();
-					tm.Language = AdaptSet->GetLanguage();
+					tm.LanguageTagRFC5646 = AdaptSet->GetLanguageTag();
 					tm.Kind = i==0 ? TEXT("main") : TEXT("translation");
 
 					for(int32 j=0, jMax=AdaptSet->GetNumberOfRepresentations(); j<jMax; ++j)
diff --git a/Engine/Plugins/Media/ElectraPlayer/Source/ElectraPlayerRuntime/Private/Runtime/Player/mpegaudio/StreamReaderMPEGAudio.cpp b/Engine/Plugins/Media/ElectraPlayer/Source/ElectraPlayerRuntime/Private/Runtime/Player/mpegaudio/StreamReaderMPEGAudio.cpp
--- a/Engine/Plugins/Media/ElectraPlayer/Source/ElectraPlayerRuntime/Private/Runtime/Player/mpegaudio/StreamReaderMPEGAudio.cpp
+++ b/Engine/Plugins/Media/ElectraPlayer/Source/ElectraPlayerRuntime/Private/Runtime/Player/mpegaudio/StreamReaderMPEGAudio.cpp
@@ -10,10 +10,10 @@
 #include "Utilities/Utilities.h"
 #include "Utilities/TimeUtilities.h"
 #include "Utilities/UtilsMP4.h"
+#include "Utilities/BCP47-Helpers.h"
 #include "Utils/MPEG/ElectraUtilsMPEGAudio.h"
 #include "HAL/LowLevelMemTracker.h"
 #include "ElectraPlayerPrivate.h"
-#include "Player/Manifest.h"
 #include "Player/PlayerStreamReader.h"
 #include "Player/AdaptiveStreamingPlayerABR.h"
 #include "Player/mpegaudio/ManifestMPEGAudio.h"
@@ -553,7 +553,7 @@ void FStreamReaderMPEGAudio::HandleRequest()
 	TSharedPtrTS<FBufferSourceInfo> BufferSourceInfo = MakeSharedTS<FBufferSourceInfo>();
 	BufferSourceInfo->PeriodID = ds.MediaAssetID;
 	BufferSourceInfo->PeriodAdaptationSetID = ds.MediaAssetID + TEXT(".") + ds.AdaptationSetID;
-	BufferSourceInfo->Language = TEXT("und");
+	BCP47::ParseRFC5646Tag(BufferSourceInfo->LanguageTag, FString(TEXT("und")));
 	BufferSourceInfo->Codec = TEXT("mp4a.6b");
 	BufferSourceInfo->HardIndex = 0;
 	BufferSourceInfo->PlaybackSequenceID = Request->GetPlaybackSequenceID();
diff --git a/Engine/Plugins/Media/ElectraPlayer/Source/ElectraPlayerRuntime/Private/Runtime/StreamAccessUnitBuffer.h b/Engine/Plugins/Media/ElectraPlayer/Source/ElectraPlayerRuntime/Private/Runtime/StreamAccessUnitBuffer.h
--- a/Engine/Plugins/Media/ElectraPlayer/Source/ElectraPlayerRuntime/Private/Runtime/StreamAccessUnitBuffer.h
+++ b/Engine/Plugins/Media/ElectraPlayer/Source/ElectraPlayerRuntime/Private/Runtime/StreamAccessUnitBuffer.h
@@ -41,8 +41,8 @@ namespace Electra
 		// Identifies the period and track (adaptation set) this data is originating from.
 		FString		PeriodAdaptationSetID;
 		// Partial track metadata. See FTrackMetadata.
+		BCP47::FLanguageTag LanguageTag;
 		FString		Kind;
-		FString		Language;
 		FString		Codec;
 
 		// Internal hard index, used for multiplexed streams.
diff --git a/Engine/Plugins/Media/ElectraPlayer/Source/ElectraPlayerRuntime/Private/Runtime/StreamTypes.h b/Engine/Plugins/Media/ElectraPlayer/Source/ElectraPlayerRuntime/Private/Runtime/StreamTypes.h
--- a/Engine/Plugins/Media/ElectraPlayer/Source/ElectraPlayerRuntime/Private/Runtime/StreamTypes.h
+++ b/Engine/Plugins/Media/ElectraPlayer/Source/ElectraPlayerRuntime/Private/Runtime/StreamTypes.h
@@ -7,6 +7,7 @@
 #include "ParameterDictionary.h"
 #include "MediaVideoDecoderOutput.h"
 #include "MediaURLType.h"
+#include "Utilities/BCP47-Helpers.h"
 
 namespace Electra
 {
@@ -524,16 +525,16 @@ namespace Electra
 			return NumberOfAudioObjects;
 		}
 
-		void SetStreamLanguageCode(const FString& InStreamLanguageCode)
+		void SetStreamLanguageTag(const BCP47::FLanguageTag& InStreamLanguageTag)
 		{
-			StreamLanguageCode = InStreamLanguageCode;
+			StreamLanguageTag = InStreamLanguageTag;
 		}
-
+/*
 		const FString& GetStreamLanguageCode() const
 		{
 			return StreamLanguageCode;
 		}
-
+*/
 		void SetCodecSpecificData(const TArray<uint8>& InCSD)
 		{
 			CSD = InCSD;
@@ -586,7 +587,7 @@ namespace Electra
 			AspectRatio.Clear();
 			FrameRate = FTimeFraction::GetInvalid();
 			ProfileLevel.Clear();
-			StreamLanguageCode.Empty();
+			StreamLanguageTag.Empty();
 			Bitrate = 0;
 			SampleRate = 0;
 			NumChannels = 0;
@@ -604,7 +605,7 @@ namespace Electra
 
 		bool Equals(const FStreamCodecInformation& Other) const
 		{
-			if (StreamType == Other.StreamType && Codec == Other.Codec && CodecSpecifier.Equals(Other.CodecSpecifier) && StreamLanguageCode.Equals(Other.StreamLanguageCode))
+			if (StreamType == Other.StreamType && Codec == Other.Codec && CodecSpecifier.Equals(Other.CodecSpecifier) && StreamLanguageTag.Get(true, true, true, false, false, false).Equals(Other.StreamLanguageTag.Get(true, true, true, false, false, false)))
 			{
 				if (StreamType == EStreamType::Video)
 				{
@@ -707,7 +708,7 @@ namespace Electra
 		FAspectRatio	AspectRatio;				//!< Aspect ratio, if this is a video stream
 		FTimeFraction	FrameRate;					//!< Frame rate, if this is a video stream
 		FProfileLevel	ProfileLevel;
-		FString			StreamLanguageCode;			//!< Language code from the stream itself, if present.
+		BCP47::FLanguageTag StreamLanguageTag;
 		int32			Bitrate;
 		int32			SampleRate;					//!< Decoded sample rate, if this is an audio stream.
 		int32			NumChannels;				//!< Number of decoded channels, if this is an audio stream.
@@ -766,9 +767,9 @@ namespace Electra
 	 */
 	struct FTrackMetadata
 	{
+		BCP47::FLanguageTag LanguageTagRFC5646;
 		FString ID;
 		FString Kind;
-		FString Language;		// ISO-639-1 language code
 		FString Label;
 
 		TArray<FStreamMetadata> StreamDetails;
@@ -776,7 +777,7 @@ namespace Electra
 		int32 HighestBandwidth = 0;
 		bool Equals(const FTrackMetadata& Other) const
 		{
-			bool bEquals = ID.Equals(Other.ID) && Kind.Equals(Other.Kind) && Language.Equals(Other.Language);
+			bool bEquals = ID.Equals(Other.ID) && Kind.Equals(Other.Kind) && LanguageTagRFC5646.Get(true, true, true, false, false, false).Equals(Other.LanguageTagRFC5646.Get(true, true, true, false, false, false));
 			if (bEquals)
 			{
 				if (StreamDetails.Num() == Other.StreamDetails.Num())
@@ -802,14 +803,13 @@ namespace Electra
 	 */
 	struct FStreamSelectionAttributes
 	{
+		// Used for video, audio and subtitles or captions. This is a comma separated string of RFC-4647 language ranges
+		// to match the RFC-5646 language tag provided by the stream metadata.
+		TOptional<FString> Language_RFC4647;
+
 		// Primarily used for audio selection. Should typically be set to "main" or left unset.
 		TOptional<FString> Kind;
 
-		// Used for audio and subtitles or captions. The language must be a valid two letter ISO-639 (-1, -2 or -3) code
-		// for which a two-letter ISO-639-1 equivalent exists. This is the equivalent of the primary language subtag of
-		// RFC-4646 and of RFC-5646
-		TOptional<FString> Language_ISO639;
-
 		// Preferred codec. Typically set to ensure the same track remains selected after a seek in case the same content
 		// is provided with different formats. See GetCodecName()
 		TOptional<FString> Codec;
@@ -818,27 +818,25 @@ namespace Electra
 		// of the tracks as they are found. If the index is invalid the selection rules for kind and language are applied.
 		TOptional<int32> OverrideIndex;
 
-		virtual ~FStreamSelectionAttributes() = default;
-
-		virtual bool IsSet() const
+		bool IsSet() const
 		{
-			return Kind.IsSet() || Language_ISO639.IsSet() || Codec.IsSet() || OverrideIndex.IsSet();
+			return Kind.IsSet() || Language_RFC4647.IsSet() || Codec.IsSet() || OverrideIndex.IsSet();
 		}
 
-		virtual bool IsCompatibleWith(const FStreamSelectionAttributes& Other)
+		bool IsCompatibleWith(const FStreamSelectionAttributes& Other)
 		{
 			if (OverrideIndex.IsSet() && Other.OverrideIndex.IsSet() && OverrideIndex.GetValue() >= 0 && Other.OverrideIndex.GetValue() >= 0 && OverrideIndex.GetValue() != Other.OverrideIndex.GetValue())
 			{
 				return false;
 			}
-			FString Kind1 = Kind.IsSet() ? Kind.GetValue() : FString();
-			FString Kind2 = Other.Kind.IsSet() ? Other.Kind.GetValue() : FString();
+			FString Kind1 = Kind.Get(FString());
+			FString Kind2 = Other.Kind.Get(FString());
 
-			FString Lang1 = Language_ISO639.IsSet() ? Language_ISO639.GetValue() : FString();
-			FString Lang2 = Other.Language_ISO639.IsSet() ? Other.Language_ISO639.GetValue() : FString();
+			FString Lang1 = Language_RFC4647.Get(FString());
+			FString Lang2 = Other.Language_RFC4647.Get(FString());
 
-			FString Codec1 = Codec.IsSet() ? Codec.GetValue() : FString();
-			FString Codec2 = Other.Codec.IsSet() ? Other.Codec.GetValue() : FString();
+			FString Codec1 = Codec.Get(FString());
+			FString Codec2 = Other.Codec.Get(FString());
 
 			if (Kind1.IsEmpty() || Kind2.IsEmpty() || Kind1.Equals(Kind2))
 			{
@@ -847,23 +845,28 @@ namespace Electra
 			return false;
 		}
 
-		virtual void Reset()
+		void Reset()
 		{
 			Kind.Reset();
-			Language_ISO639.Reset();
+			Language_RFC4647.Reset();
 			Codec.Reset();
 			OverrideIndex.Reset();
 		}
-		virtual void UpdateWith(const FString& InKind, const FString& InLanguage, const FString& InCodec, int32 InOverrideIndex)
+		void UpdateWith(const FString& InKind, const BCP47::FLanguageTag& InLanguageTag, const FString& InCodec, int32 InOverrideIndex)
 		{
 			Reset();
 			if (!InKind.IsEmpty())
 			{
 				Kind = InKind;
 			}
-			if (!InLanguage.IsEmpty())
+			FString Lang5646 = InLanguageTag.Get();
+			if (!Lang5646.IsEmpty())
+			{
+				Language_RFC4647 = MoveTemp(Lang5646);
+			}
+			else
 			{
-				Language_ISO639 = InLanguage;
+				Language_RFC4647.Reset();
 			}
 			if (!InCodec.IsEmpty())
 			{
@@ -874,17 +877,17 @@ namespace Electra
 				OverrideIndex = InOverrideIndex;
 			}
 		}
-		virtual void UpdateIfOverrideSet(const FString& InKind, const FString& InLanguage, const FString& InCodec)
+		void UpdateIfOverrideSet(const FString& InKind, const BCP47::FLanguageTag& InLanguageTag, const FString& InCodec)
 		{
 			if (OverrideIndex.IsSet())
 			{
 				ClearOverrideIndex();
 				Kind = InKind;
-				Language_ISO639 = InLanguage;
+				Language_RFC4647 = InLanguageTag.Get();
 				Codec = InCodec;
 			}
 		}
-		virtual void ClearOverrideIndex()
+		void ClearOverrideIndex()
 		{
 			OverrideIndex.Reset();
 		}
diff --git a/Engine/Plugins/Media/ElectraPlayer/Source/ElectraPlayerRuntime/Public/IElectraPlayerInterface.h b/Engine/Plugins/Media/ElectraPlayer/Source/ElectraPlayerRuntime/Public/IElectraPlayerInterface.h
--- a/Engine/Plugins/Media/ElectraPlayer/Source/ElectraPlayerRuntime/Public/IElectraPlayerInterface.h
+++ b/Engine/Plugins/Media/ElectraPlayer/Source/ElectraPlayerRuntime/Public/IElectraPlayerInterface.h
@@ -156,13 +156,14 @@ public:
 	struct FStreamSelectionAttributes
 	{
 		TOptional<FString> Kind;
-		TOptional<FString> Language_ISO639;
+		TOptional<FString> Language_RFC4647;
 		TOptional<FString> Codec;
 		TOptional<int32> TrackIndexOverride;
 		void Reset()
 		{
 			Kind.Reset();
-			Language_ISO639.Reset();
+			Language_RFC4647.Reset();
+			Codec.Reset();
 			TrackIndexOverride.Reset();
 		}
 	};
diff --git a/Engine/Plugins/Media/ElectraPlayer/Source/ElectraProtron/Private/Player/ElectraProtronPlayerImpl.cpp b/Engine/Plugins/Media/ElectraPlayer/Source/ElectraProtron/Private/Player/ElectraProtronPlayerImpl.cpp
--- a/Engine/Plugins/Media/ElectraPlayer/Source/ElectraProtron/Private/Player/ElectraProtronPlayerImpl.cpp
+++ b/Engine/Plugins/Media/ElectraPlayer/Source/ElectraProtron/Private/Player/ElectraProtronPlayerImpl.cpp
@@ -1069,7 +1069,7 @@ FString FElectraProtronPlayer::FImpl::GetTrackLanguage(EMediaTrackType InTrackTy
 			{
 				if (InTrackIndex < UsableTrackArrayIndicesByType[CodecTypeIndex(ElectraProtronUtils::FCodecInfo::EType::Video)].Num())
 				{
-					return Tracks[UsableTrackArrayIndicesByType[CodecTypeIndex(ElectraProtronUtils::FCodecInfo::EType::Video)][InTrackIndex]]->MP4Track->GetCommonMetadata().Language;
+					return Tracks[UsableTrackArrayIndicesByType[CodecTypeIndex(ElectraProtronUtils::FCodecInfo::EType::Video)][InTrackIndex]]->MP4Track->GetCommonMetadata().LanguageTag.Get();
 				}
 				break;
 			}
@@ -1077,7 +1077,7 @@ FString FElectraProtronPlayer::FImpl::GetTrackLanguage(EMediaTrackType InTrackTy
 			{
 				if (InTrackIndex < UsableTrackArrayIndicesByType[CodecTypeIndex(ElectraProtronUtils::FCodecInfo::EType::Audio)].Num())
 				{
-					return Tracks[UsableTrackArrayIndicesByType[CodecTypeIndex(ElectraProtronUtils::FCodecInfo::EType::Audio)][InTrackIndex]]->MP4Track->GetCommonMetadata().Language;
+					return Tracks[UsableTrackArrayIndicesByType[CodecTypeIndex(ElectraProtronUtils::FCodecInfo::EType::Audio)][InTrackIndex]]->MP4Track->GetCommonMetadata().LanguageTag.Get();
 				}
 				break;
 			}
diff --git a/Engine/Plugins/Media/ElectraUtil/Source/ElectraBase/Private/Utilities/BCP47-Helpers.cpp b/Engine/Plugins/Media/ElectraUtil/Source/ElectraBase/Private/Utilities/BCP47-Helpers.cpp
new file mode 100644
--- /dev/null
+++ b/Engine/Plugins/Media/ElectraUtil/Source/ElectraBase/Private/Utilities/BCP47-Helpers.cpp
@@ -0,0 +1,401 @@
+// Copyright Epic Games, Inc. All Rights Reserved.
+
+#include "Utilities/BCP47-Helpers.h"
+#include "Utilities/ISO639-Map.h"
+#include "Internationalization/Regex.h"
+
+namespace Electra
+{
+	namespace BCP47
+	{
+		/*
+			Note: If, at some point, we wanted to canonicalize the language tag we can use the
+			      IANA database as a source of information
+					https://www.iana.org/assignments/language-subtag-registry/language-subtag-registry
+
+				  The format is explained in RFC-5646.
+		*/
+
+		static const TArray<const FString> Irregulars
+		{
+			TEXT("en-GB-oed"),
+			TEXT("i-ami"),
+			TEXT("i-bnn"),
+			TEXT("i-default"),
+			TEXT("i-enochian"),
+			TEXT("i-hak"),
+			TEXT("i-klingon"),
+			TEXT("i-lux"),
+			TEXT("i-mingo"),
+			TEXT("i-navajo"),
+			TEXT("i-pwn"),
+			TEXT("i-tao"),
+			TEXT("i-tay"),
+			TEXT("i-tsu"),
+			TEXT("sgn-BE-FR"),
+			TEXT("sgn-BE-NL"),
+			TEXT("sgn-CH-DE")
+		};
+
+		static const TArray<const FString> Regulars
+		{
+			TEXT("art-lojban"),
+			TEXT("cel-gaulish"),
+			TEXT("no-bok"),
+			TEXT("no-nyn"),
+			TEXT("zh-guoyu"),
+			TEXT("zh-hakka"),
+			TEXT("zh-min"),
+			TEXT("zh-min-nan"),
+			TEXT("zh-xiang")
+		};
+
+		namespace Regexes
+		{
+			static const FRegexPattern& WholePrivateUse()
+			{
+				static const FRegexPattern pu(TEXT(R"(^([xX](?:-[a-zA-Z0-9]{1,8})+)$)"));
+				return pu;
+			}
+			static const FRegexPattern& Language()
+			{
+				static const FRegexPattern la(TEXT(R"(^((?:(?:[a-zA-Z]{2,3}(?:(?:-[a-zA-Z]{3}){0,3})|(?:[a-zA-Z]{4})|(?:[a-zA-Z]{5,8}))(?=-|$))))"));
+				return la;
+			}
+			static const FRegexPattern& Script()
+			{
+				static const FRegexPattern sc(TEXT(R"(^-([a-zA-Z]{4})(?=-|$))"));
+				return sc;
+			}
+			static const FRegexPattern& Region()
+			{
+				static const FRegexPattern re(TEXT(R"(^-([a-zA-Z]{2}|[0-9]{3})(?=-|$))"));
+				return re;
+			}
+			static const FRegexPattern& Variant()
+			{
+				static const FRegexPattern va(TEXT(R"(^-([a-zA-Z0-9]{5,8}(?=-|$)|(?:[0-9][a-zA-Z0-9]{3}(?=-|$))))"));
+				return va;
+			}
+			static const FRegexPattern& Extension()
+			{
+				static const FRegexPattern ex(TEXT(R"(^-([0-9a-wyzA-WYZ](?:(?:-[a-zA-Z0-9]{2,8})+)(?=-|$)))"));
+				return ex;
+			}
+			static const FRegexPattern& PrivateUse()
+			{
+				static const FRegexPattern pu(TEXT(R"(^-([xX](?:-[a-zA-Z0-9]{1,8})+))"));
+				return pu;
+			}
+		}
+
+
+		static bool ParseInternal(FLanguageTag& OutTag, FString& OutError, const FString& InRFC5646)
+		{
+			// First check if this is a private use tag as a whole
+			FRegexMatcher WholePrivateUse(Regexes::WholePrivateUse(), InRFC5646);
+			if (WholePrivateUse.FindNext())
+			{
+				OutTag.FullLanguage = OutTag.PrimaryLanguage = WholePrivateUse.GetCaptureGroup(1);
+				return true;
+			}
+			// Then check if it is a grandfathered tag.
+			auto CheckGrandfathered = [](FLanguageTag& Out, const FString& InTag, const TArray<const FString>& InList) -> bool
+			{
+				for(auto& i : InList)
+				{
+					// If it matches set the language from the list so the capitalization is as it should be,.
+					if (i.Equals(InTag, ESearchCase::IgnoreCase))
+					{
+						Out.FullLanguage = Out.PrimaryLanguage = i;
+						return true;
+					}
+				}
+				return false;
+			};
+			if (CheckGrandfathered(OutTag, InRFC5646, Irregulars) || CheckGrandfathered(OutTag, InRFC5646, Regulars))
+			{
+				return true;
+			}
+
+			enum class ESubPart
+			{
+				None,
+				Language,
+				Script,
+				Region,
+				Variant,
+				Extension,
+				PrivateUse
+			};
+			static const TCHAR* const SubPartNames[] =
+			{
+				TEXT("none"),
+				TEXT("language"),
+				TEXT("script"),
+				TEXT("region"),
+				TEXT("variant"),
+				TEXT("extension"),
+				TEXT("privateuse")
+			};
+			ESubPart LastSuccessfulSubPart = ESubPart::None;
+			int32 ParsePos = 0;
+			FString Remainder(InRFC5646);
+			FRegexMatcher Language(Regexes::Language(), Remainder);
+			if (!Language.FindNext())
+			{
+				OutError = TEXT("Language not found at beginning");
+				return false;
+			}
+			OutTag.FullLanguage = Language.GetCaptureGroup(1);
+			if (OutTag.FullLanguage.Len() == 4)
+			{
+				OutError = TEXT("Four letter language is reserved for future use");
+				return false;
+			}
+			Remainder.MidInline(OutTag.FullLanguage.Len());
+			ParsePos += OutTag.FullLanguage.Len();
+			LastSuccessfulSubPart = ESubPart::Language;
+			int32 PrimLangPos;
+			if (OutTag.FullLanguage.FindChar(TCHAR('-'), PrimLangPos))
+			{
+				// Map the primary language to the shortest possible one
+				OutTag.PrimaryLanguage = ISO639::MapTo639_1(OutTag.FullLanguage.Mid(0, PrimLangPos));
+				OutTag.ExtendedLanguage = OutTag.FullLanguage.Mid(PrimLangPos + 1);
+				// Then reassemble the full language again.
+				OutTag.FullLanguage = FString::Printf(TEXT("%s-%s"), *OutTag.PrimaryLanguage, *OutTag.ExtendedLanguage);
+			}
+			else
+			{
+				OutTag.FullLanguage = ISO639::MapTo639_1(OutTag.FullLanguage);
+				OutTag.PrimaryLanguage = OutTag.FullLanguage;
+			}
+			// Try script
+			FRegexMatcher Script(Regexes::Script(), Remainder);
+			if (Script.FindNext())
+			{
+				OutTag.Script = Script.GetCaptureGroup(1);
+				OutTag.Script[0] = FChar::ToUpper(OutTag.Script[0]);
+				Remainder.MidInline(1 + OutTag.Script.Len());
+				ParsePos += 1 + OutTag.Script.Len();
+				LastSuccessfulSubPart = ESubPart::Script;
+			}
+			// Try region
+			FRegexMatcher Region(Regexes::Region(), Remainder);
+			if (Region.FindNext())
+			{
+				OutTag.Region = Region.GetCaptureGroup(1).ToUpper();
+				Remainder.MidInline(1 + OutTag.Region.Len());
+				ParsePos += 1 + OutTag.Region.Len();
+				LastSuccessfulSubPart = ESubPart::Region;
+			}
+			// Now see if there are any variants
+			while(1)
+			{
+				FRegexMatcher Variant(Regexes::Variant(), Remainder);
+				if (Variant.FindNext())
+				{
+					FString V(Variant.GetCaptureGroup(1));
+					if (OutTag.Variants.ContainsByPredicate([&v=V](const FString& e){ return e.Equals(v, ESearchCase::IgnoreCase); }))
+					{
+						OutError = FString::Printf(TEXT("Variant %s appears more than once"), *V);
+						return false;
+					}
+					OutTag.Variants.Emplace(V);
+					Remainder.MidInline(1 + V.Len());
+					ParsePos += 1 + V.Len();
+					LastSuccessfulSubPart = ESubPart::Variant;
+				}
+				else
+				{
+					break;
+				}
+			}
+			// Extensions?
+			while(1)
+			{
+				FRegexMatcher Extension(Regexes::Extension(), Remainder);
+				if (Extension.FindNext())
+				{
+					OutTag.Extensions.Emplace(Extension.GetCaptureGroup(1));
+					Remainder.MidInline(1 + OutTag.Extensions.Last().Len());
+					ParsePos += 1 + OutTag.Extensions.Last().Len();
+					LastSuccessfulSubPart = ESubPart::Extension;
+				}
+				else
+				{
+					break;
+				}
+			}
+			// Private use?
+			FRegexMatcher PrivateUse(Regexes::PrivateUse(), Remainder);
+			if (PrivateUse.FindNext())
+			{
+				OutTag.PrivateUse = PrivateUse.GetCaptureGroup(1);
+				Remainder.MidInline(1 + OutTag.PrivateUse.Len());
+				ParsePos += 1 + OutTag.PrivateUse.Len();
+				LastSuccessfulSubPart = ESubPart::PrivateUse;
+			}
+			// We need to have consumed the entire language tag for parsing to be
+			// successful. If there is still something left then the tag is malformed.
+			if (!Remainder.IsEmpty())
+			{
+				OutError = FString::Printf(TEXT("Error after %s sub tag at position %d: \"%s\""), SubPartNames[(int32)LastSuccessfulSubPart], ParsePos+1, *Remainder);
+				return false;
+			}
+			return true;
+		}
+
+
+
+		bool ParseRFC5646Tag(FLanguageTag& OutTag, const FString& InRFC5646)
+		{
+			FString ErrorMsg;
+			return ParseInternal(OutTag, ErrorMsg, InRFC5646.ToLower());
+		}
+
+		TArray<int32> FindExtendedFilteringMatch(const TArray<FLanguageTag>& InTagsToCheck, const FString& InRFC4647Ranges)
+		{
+			TArray<FString> RangesToTest;
+			InRFC4647Ranges.ParseIntoArray(RangesToTest, TEXT(","), true);
+			// No test range, no result.
+			if (RangesToTest.IsEmpty())
+			{
+				return TArray<int32>();
+			}
+			// Loop over the candidates
+			TArray<int32> ResultIndices;
+			for(int32 nCand=0; nCand<InTagsToCheck.Num(); ++nCand)
+			{
+				// Step 1: break apart the language tag and the language range on the hyphens.
+				TArray<FString> CandidateParts;
+				InTagsToCheck[nCand].Get().ParseIntoArray(CandidateParts, TEXT("-"), true);
+				// Test each language range in turn.
+				bool bFoundMatch = false;
+				for(int32 nRange=0; !bFoundMatch && nRange<RangesToTest.Num(); ++nRange)
+				{
+					TArray<FString> RangeParts;
+					RangesToTest[nRange].ParseIntoArray(RangeParts, TEXT("-"), true);
+					check(!RangeParts.IsEmpty());	// can't be empty since we culled empty ranges on entry.
+					// Step 2: check language part
+					bool bFirstIsWildcard = RangeParts[0].Equals(TEXT("*"));
+					if (CandidateParts.IsEmpty())
+					{
+						if (bFirstIsWildcard)
+						{
+							bFoundMatch = true;
+							ResultIndices.AddUnique(nCand);
+						}
+						continue;
+					}
+					else if (bFirstIsWildcard || RangeParts[0].Equals(CandidateParts[0], ESearchCase::IgnoreCase))
+					{
+						int32 candPartIdx = 1;
+						int32 rngPartIdx = 1;
+						bool bMatches = true;
+						// Step 3:
+						while(rngPartIdx < RangeParts.Num())
+						{
+							// 3A:
+							if (RangeParts[rngPartIdx].Equals(TEXT("*")))
+							{
+								++rngPartIdx;
+								continue;
+							}
+							// 3B:
+							if (candPartIdx >= CandidateParts.Num())
+							{
+								bMatches = false;
+								break;
+							}
+							// 3C:
+							if (RangeParts[rngPartIdx].Equals(CandidateParts[candPartIdx], ESearchCase::IgnoreCase))
+							{
+								++rngPartIdx;
+								++candPartIdx;
+								continue;
+							}
+							// 3D:
+							if (CandidateParts[candPartIdx].Len() == 1)
+							{
+								bMatches = false;
+								break;
+							}
+							// 3E:
+							++candPartIdx;
+						}
+						if (bMatches)
+						{
+							bFoundMatch = true;
+							ResultIndices.AddUnique(nCand);
+						}
+					}
+				}
+			}
+
+			// If filtering produced no results try "lookup"
+			if (ResultIndices.IsEmpty())
+			{
+				int32 BestMatchPos = 0;
+				int32 BestMatchIndex = -1;
+				// Try each language range in priority order.
+				for(auto& testRange : RangesToTest)
+				{
+					// Cannot use language range containing wildcard
+					if (testRange.Contains(TEXT("*")))
+					{
+						continue;
+					}
+					// Parse the languate range as a tag. If this fails we ignore it.
+					FLanguageTag testTag;
+					FString parseError;
+					if (!ParseInternal(testTag, parseError, testRange))
+					{
+						continue;
+					}
+
+					// Check the parsed language range against the list of given language tags.
+					for(int32 nCand=0; nCand<InTagsToCheck.Num(); ++nCand)
+					{
+						// Compare the sub tags element-wise.
+						int32 MatchPos = 0;
+						if (testTag.PrimaryLanguage.Equals(InTagsToCheck[nCand].PrimaryLanguage, ESearchCase::IgnoreCase))
+						{
+							++MatchPos;
+							if (testTag.ExtendedLanguage.Equals(InTagsToCheck[nCand].ExtendedLanguage, ESearchCase::IgnoreCase))
+							{
+								++MatchPos;
+								if (testTag.Script.Equals(InTagsToCheck[nCand].Script, ESearchCase::IgnoreCase))
+								{
+									++MatchPos;
+									if (testTag.Region.Equals(InTagsToCheck[nCand].Region, ESearchCase::IgnoreCase))
+									{
+										++MatchPos;
+									}
+								}
+							}
+						}
+						// A (better) match?
+						if (MatchPos > BestMatchPos)
+						{
+							BestMatchPos = MatchPos;
+							BestMatchIndex = nCand;
+						}
+					}
+
+					// If the language range matched one of the given language tags we stop.
+					// The assumption still is that the language ranges (if more than one) are given
+					// in most descriptive to least descriptive order.
+					if (BestMatchIndex >= 0)
+					{
+						ResultIndices.Emplace(BestMatchIndex);
+						break;
+					}
+				}
+			}
+
+			return ResultIndices;
+		}
+
+	}
+}
diff --git a/Engine/Plugins/Media/ElectraUtil/Source/ElectraBase/Private/Utilities/MP4Boxes/MP4Boxes.cpp b/Engine/Plugins/Media/ElectraUtil/Source/ElectraBase/Private/Utilities/MP4Boxes/MP4Boxes.cpp
--- a/Engine/Plugins/Media/ElectraUtil/Source/ElectraBase/Private/Utilities/MP4Boxes/MP4Boxes.cpp
+++ b/Engine/Plugins/Media/ElectraUtil/Source/ElectraBase/Private/Utilities/MP4Boxes/MP4Boxes.cpp
@@ -3,6 +3,7 @@
 #include "Utilities/MP4Boxes/MP4Boxes.h"
 #include "Utilities/UtilitiesMP4.h"
 #include "Utilities/ISO639-Map.h"
+#include "Utilities/BCP47-Helpers.h"
 #include "ElectraBaseModule.h"	// for log category `LogElectraBase`
 
 namespace Electra
@@ -10,7 +11,6 @@ namespace Electra
 	namespace UtilitiesMP4
 	{
 
-
 		void FMP4BoxBase::ProcessBoxChildrenRecursively(FMP4AtomReaderHelper& InOutReader, const FMP4BoxInfo& InCurrentBoxInfo)
 		{
 			FMP4BoxInfo bi;
@@ -267,7 +267,7 @@ namespace Electra
 
 		struct FMP4BoxMDHD::FParsed
 		{
-			FString Lang639_1;
+			BCP47::FLanguageTag LanguageTag;
 			int64 Duration = 0;
 			uint32 Flags = 0;
 			uint32 Timescale = 0;
@@ -302,7 +302,7 @@ namespace Electra
 			ar.Read(Value16);
 			char Language[3] { (char)(0x60 + ((Value16 & 0x7c00) >> 10)), (char)(0x60 + ((Value16 & 0x03e0) >> 5)), (char)(0x60 + (Value16 & 0x001f)) };
 			// Try to map the ISO-639-2T language code to the shorter ISO-639-1 code if possible.
-			Parsed->Lang639_1 = ISO639::MapTo639_1(FString::ConstructFromPtrSize(Language, UE_ARRAY_COUNT(Language)));
+			BCP47::ParseRFC5646Tag(Parsed->LanguageTag, FString::ConstructFromPtrSize(Language, UE_ARRAY_COUNT(Language)));
 		}
 
 		FTimeFraction FMP4BoxMDHD::GetDuration()
@@ -315,10 +315,10 @@ namespace Electra
 			ParseIfRequired();
 			return Parsed->Timescale;
 		}
-		const FString& FMP4BoxMDHD::GetLanguage()
+		const BCP47::FLanguageTag& FMP4BoxMDHD::GetLanguageTag()
 		{
 			ParseIfRequired();
-			return Parsed->Lang639_1;
+			return Parsed->LanguageTag;
 		}
 
 /****************************************************************************************************************************************************/
@@ -360,7 +360,7 @@ namespace Electra
 
 		struct FMP4BoxELNG::FParsed
 		{
-			FString Language;
+			BCP47::FLanguageTag LanguageTag;
 			uint32 Flags = 0;
 			uint8 Version = 0;
 		};
@@ -373,13 +373,15 @@ namespace Electra
 			Parsed = MakePimpl<FParsed>();
 			FMP4AtomReaderHelper ar(BoxInfo.Data);
 			ar.ReadVersionAndFlags(Parsed->Version, Parsed->Flags);
-			ar.ReadStringUTF8(Parsed->Language, ar.GetNumBytesRemaining());
+			FString Lang;
+			ar.ReadStringUTF8(Lang, ar.GetNumBytesRemaining());
+			BCP47::ParseRFC5646Tag(Parsed->LanguageTag, Lang);
 		}
 
-		FString FMP4BoxELNG::GetLanguage()
+		const BCP47::FLanguageTag& FMP4BoxELNG::GetLanguageTag()
 		{
 			ParseIfRequired();
-			return Parsed->Language;
+			return Parsed->LanguageTag;
 		}
 
 /****************************************************************************************************************************************************/
diff --git a/Engine/Plugins/Media/ElectraUtil/Source/ElectraBase/Private/Utilities/MP4Boxes/MP4Track.cpp b/Engine/Plugins/Media/ElectraUtil/Source/ElectraBase/Private/Utilities/MP4Boxes/MP4Track.cpp
--- a/Engine/Plugins/Media/ElectraUtil/Source/ElectraBase/Private/Utilities/MP4Boxes/MP4Track.cpp
+++ b/Engine/Plugins/Media/ElectraUtil/Source/ElectraBase/Private/Utilities/MP4Boxes/MP4Track.cpp
@@ -171,12 +171,12 @@ namespace Electra
 				CommonMetadata.HandlerName = HdlrBox->GetHandlerName();
 			}
 			// Get language from `mdhd` box first.
-			CommonMetadata.Language = MdhdBox->GetLanguage();
+			CommonMetadata.LanguageTag = MdhdBox->GetLanguageTag();
 			// Then, if there is an `elng` box that gives better information get it from there.
 			TSharedPtr<FMP4BoxELNG, ESPMode::ThreadSafe> ElngBox = MdiaBox->FindBoxRecursive<Electra::UtilitiesMP4::FMP4BoxELNG>(Electra::UtilitiesMP4::MakeBoxAtom('e','l','n','g'), 0);
 			if (ElngBox.IsValid())
 			{
-				CommonMetadata.Language = ElngBox->GetLanguage();
+				CommonMetadata.LanguageTag = ElngBox->GetLanguageTag();
 			}
 
 			bHasBeenPrepared = true;
diff --git a/Engine/Plugins/Media/ElectraUtil/Source/ElectraBase/Public/Utilities/BCP47-Helpers.h b/Engine/Plugins/Media/ElectraUtil/Source/ElectraBase/Public/Utilities/BCP47-Helpers.h
new file mode 100644
--- /dev/null
+++ b/Engine/Plugins/Media/ElectraUtil/Source/ElectraBase/Public/Utilities/BCP47-Helpers.h
@@ -0,0 +1,132 @@
+// Copyright Epic Games, Inc. All Rights Reserved.
+
+#pragma once
+
+#include "CoreTypes.h"
+#include "Containers/UnrealString.h"
+
+namespace Electra
+{
+
+	namespace BCP47
+	{
+
+		/**
+		 * A parsed RFC-5646 language tag.
+		 */
+		struct FLanguageTag
+		{
+			// The full language string if it consists of more than just the primary language.
+			FString FullLanguage;
+			// The primary language only.
+			FString PrimaryLanguage;
+			// Any extended language sub parts (up to three) separated by hyphens.
+			FString ExtendedLanguage;
+			// The script sub part, if given.
+			FString Script;
+			// The region sub part, if given.
+			FString Region;
+			// The variant sub parts, if given.
+			TArray<FString> Variants;
+			// The extension sub parts, if given.
+			TArray<FString> Extensions;
+			// The private use part, if given.
+			FString PrivateUse;
+
+			void Empty()
+			{
+				FullLanguage.Empty();
+				PrimaryLanguage.Empty();
+				ExtendedLanguage.Empty();
+				Script.Empty();
+				Region.Empty();
+				Variants.Empty();
+				Extensions.Empty();
+				PrivateUse.Empty();
+			}
+
+			FString Get(bool bAddExtended, bool bAddScript, bool bAddRegion, bool bAddVariants, bool bAddExtensions, bool bAddPrivateUse) const
+			{
+				FString r(PrimaryLanguage);
+				if (bAddExtended && !ExtendedLanguage.IsEmpty())
+				{
+					r.AppendChar(TCHAR('-'));
+					r.Append(ExtendedLanguage);
+				}
+				if (bAddScript && !Script.IsEmpty())
+				{
+					r.AppendChar(TCHAR('-'));
+					r.Append(Script);
+				}
+				if (bAddRegion && !Region.IsEmpty())
+				{
+					r.AppendChar(TCHAR('-'));
+					r.Append(Region);
+				}
+				if (bAddVariants && !Variants.IsEmpty())
+				{
+					for(auto& v : Variants)
+					{
+						r.AppendChar(TCHAR('-'));
+						r.Append(v);
+					}
+				}
+				if (bAddExtensions && !Extensions.IsEmpty())
+				{
+					for(auto& e : Extensions)
+					{
+						r.AppendChar(TCHAR('-'));
+						r.Append(e);
+					}
+				}
+				if (bAddPrivateUse && !PrivateUse.IsEmpty())
+				{
+					r.AppendChar(TCHAR('-'));
+					r.Append(PrivateUse);
+				}
+				return r;
+			}
+
+			FString Get() const
+			{
+				return Get(true, true, true, true, true, true);
+			}
+		};
+
+		/*
+			Parse a single RFC-5646 language tag into its components.
+			This will change the language code to the shorter 2 letter ISO-639-1 code if possible
+			and change the capitalization of the elements to their recommended case.
+			It does NOT canonicalize the tag.
+		*/
+		ELECTRABASE_API bool ParseRFC5646Tag(FLanguageTag& OutTag, const FString& InRFC5646);
+
+
+		/*
+			Checks an RFC-4647 language priority list for a match against the given language tags
+			through RFC-4647 section 3.3.2 "extended filtering".
+			Returns the indices into the given language tags that match the language priority list.
+
+			If this filtering produces no match, "lookup" is performed (see RFC-4647 section 3.4)
+			with the following differences:
+			  - only language ranges not containing a '*' wildcard are considered because the language
+			    range must be parseable as a language tag for element-wise comparison.
+			  - instead of progressibely truncating the language range from the end the search
+			    is performed forward and stops on the first mismatch.
+			  - variants, extensions and private use tags are NOT considered!
+			  - there is no "default" value to fall back on. The result list will be empty if
+			    lookup fails to produce a match.
+
+			The point of performing an additional "lookup" step is to return a better-than-nothing
+			result if the language range is over-specified with regard to the given language tags.
+			For example, if the language tag is "es" a filter search with only "es-419" will not return
+			a match. The search should be "es-419,es-*,es" but constructing this is sometimes difficult.
+			This is why a lookup is performed because searching "es-419" will match "es".
+
+
+			May be an empty list if there are no matches.
+		*/
+		ELECTRABASE_API TArray<int32> FindExtendedFilteringMatch(const TArray<FLanguageTag>& InTagsToCheck, const FString& InRFC4647Ranges);
+	}
+
+}
diff --git a/Engine/Plugins/Media/ElectraUtil/Source/ElectraBase/Public/Utilities/MP4Boxes/MP4BoxMetadata.h b/Engine/Plugins/Media/ElectraUtil/Source/ElectraBase/Public/Utilities/MP4Boxes/MP4BoxMetadata.h
--- a/Engine/Plugins/Media/ElectraUtil/Source/ElectraBase/Public/Utilities/MP4Boxes/MP4BoxMetadata.h
+++ b/Engine/Plugins/Media/ElectraUtil/Source/ElectraBase/Public/Utilities/MP4Boxes/MP4BoxMetadata.h
@@ -13,7 +13,7 @@ namespace Electra
 
 		struct FMP4TrackMetadataCommon
 		{
-			FString Language;
+			BCP47::FLanguageTag LanguageTag;
 			FString Name;
 			FString HandlerName;
 		};
diff --git a/Engine/Plugins/Media/ElectraUtil/Source/ElectraBase/Public/Utilities/MP4Boxes/MP4Boxes.h b/Engine/Plugins/Media/ElectraUtil/Source/ElectraBase/Public/Utilities/MP4Boxes/MP4Boxes.h
--- a/Engine/Plugins/Media/ElectraUtil/Source/ElectraBase/Public/Utilities/MP4Boxes/MP4Boxes.h
+++ b/Engine/Plugins/Media/ElectraUtil/Source/ElectraBase/Public/Utilities/MP4Boxes/MP4Boxes.h
@@ -3,6 +3,7 @@
 #pragma once
 
 #include "Utilities/MP4Boxes/MP4BoxBase.h"
+#include "Utilities/BCP47-Helpers.h"
 #include "Containers/ChunkedArray.h"
 #include "PlayerTime.h"
 
@@ -263,7 +264,7 @@ namespace Electra
 
 			ELECTRABASE_API FTimeFraction GetDuration();
 			ELECTRABASE_API uint32 GetTimescale();
-			ELECTRABASE_API const FString& GetLanguage();
+			ELECTRABASE_API const BCP47::FLanguageTag& GetLanguageTag();
 
 		protected:
 			FMP4BoxMDHD(const TWeakPtr<FMP4BoxBase, ESPMode::ThreadSafe>& InParent, const FMP4BoxInfo& InBoxInfo)
@@ -327,7 +328,7 @@ namespace Electra
 			{ return MakeShareable<FMP4BoxELNG>(new FMP4BoxELNG(InParent, InBoxInfo)); }
 			virtual ~FMP4BoxELNG() = default;
 
-			ELECTRABASE_API FString GetLanguage();
+			ELECTRABASE_API const BCP47::FLanguageTag& GetLanguageTag();
 		protected:
 			FMP4BoxELNG(const TWeakPtr<FMP4BoxBase, ESPMode::ThreadSafe>& InParent, const FMP4BoxInfo& InBoxInfo)
 				: FMP4BoxBase(InParent, InBoxInfo)
diff --git a/Engine/Source/Runtime/MediaUtils/Public/MediaPlayerOptions.h b/Engine/Source/Runtime/MediaUtils/Public/MediaPlayerOptions.h
--- a/Engine/Source/Runtime/MediaUtils/Public/MediaPlayerOptions.h
+++ b/Engine/Source/Runtime/MediaUtils/Public/MediaPlayerOptions.h
@@ -35,7 +35,9 @@ enum class EMediaPlayerOptionTrackSelectMode : uint8
 	// Let the media player choose defaults.
 	UseMediaPlayerDefaults UMETA(DisplayName="Media player selects default tracks"),
 	// Use fixed track indices as specified with MediaPlayerTrackOptions
-	UseTrackOptionIndices UMETA(DisplayName="Uses provided track indices")
+	UseTrackOptionIndices UMETA(DisplayName="Uses provided track indices"),
+	// Use language codes as specified with MediaPlayerInitialTrackLanguageSelection
+	UseLanguageCodes UMETA(DisplayName="Use provided language codes or have media player select defaults when there is no match")
 };
 
 
@@ -89,6 +91,35 @@ struct FMediaPlayerTrackOptions
 	int32 Video;
 };
 
+
+USTRUCT(BlueprintType)
+struct FMediaPlayerInitialTrackLanguageSelection
+{
+	GENERATED_BODY()
+
+	FMediaPlayerInitialTrackLanguageSelection()
+	{ }
+
+	bool operator!=(const FMediaPlayerInitialTrackLanguageSelection& Other) const
+	{ return Video != Other.Video || Audio != Other.Audio || Subtitle != Other.Subtitle || Caption != Other.Caption; }
+
+	// Language code or codes of the video track to select. Useful when the video contains burned-in subtitles or localized scenes.
+	UPROPERTY(BlueprintReadWrite, EditAnywhere, Category = "Tracks", meta=(ToolTip="Comma separated RFC 4647 language ranges in order of priority, or keep empty for default."))
+	FString Video;
+
+	// Language code or codes of the audio track to select.
+	UPROPERTY(BlueprintReadWrite, EditAnywhere, Category = "Tracks", meta=(ToolTip="Comma separated RFC 4647 language ranges in order of priority, or keep empty for default."))
+	FString Audio;
+
+	// Language code or codes of the subtitle track to select.
+	UPROPERTY(BlueprintReadWrite, EditAnywhere, Category = "Tracks", meta=(ToolTip="Comma separated RFC 4647 language ranges in order of priority, or keep empty for default."))
+	FString Subtitle;
+
+	// Language code or codes of the caption track to select. Captions may not provide language codes if carried in the video signal.
+	UPROPERTY(BlueprintReadWrite, EditAnywhere, Category = "Tracks", meta=(ToolTip="Comma separated RFC 4647 language ranges in order of priority, or keep empty for default."))
+	FString Caption;
+};
+
 USTRUCT(BlueprintType)
 struct FMediaPlayerOptions
 {
@@ -116,6 +147,7 @@ struct FMediaPlayerOptions
 	bool operator!=(const FMediaPlayerOptions& Other) const
 	{
 		return (Tracks != Other.Tracks
+			|| TracksByLanguage != Other.TracksByLanguage
 			|| TrackSelection != Other.TrackSelection
 			|| SeekTime != Other.SeekTime
 			|| SeekTimeType != Other.SeekTimeType
@@ -127,6 +159,10 @@ struct FMediaPlayerOptions
 	UPROPERTY(BlueprintReadWrite, EditAnywhere, Category = "Tracks", meta=(DisplayName="Initial track indices", ToolTip="Indices of the media tracks to select for playback", EditCondition="TrackSelection==EMediaPlayerOptionTrackSelectMode::UseTrackOptionIndices"))
 	FMediaPlayerTrackOptions Tracks;
 
+	/** Track selection by language. Not supported by all media players. */
+	UPROPERTY(BlueprintReadWrite, EditAnywhere, Category = "Tracks", meta=(DisplayName="Initial track language", ToolTip="Language of the media tracks to select for playback", EditCondition="TrackSelection==EMediaPlayerOptionTrackSelectMode::UseLanguageCodes"))
+	FMediaPlayerInitialTrackLanguageSelection TracksByLanguage;
+
 	/** How to select the initial media tracks. */
 	UPROPERTY(BlueprintReadWrite, EditAnywhere, Category = "Tracks", meta=(DisplayName="Initial track selection mode", ToolTip="How the initial media tracks for playback are selected"))
 	EMediaPlayerOptionTrackSelectMode TrackSelection;
@@ -154,20 +190,20 @@ struct FMediaPlayerOptions
 
 namespace MediaPlayerOptionValues
 {
-	inline const FName& Environment() 
+	inline const FName& Environment()
 	{ static FName OptName(TEXT("Environment")); return OptName; }
-	inline const FName& Environment_Sequencer() 
+	inline const FName& Environment_Sequencer()
 	{ static FName OptName(TEXT("Sequencer")); return OptName; }
-	inline const FName& Environment_Preview() 
+	inline const FName& Environment_Preview()
 	{ static FName OptName(TEXT("Preview")); return OptName; }
 
-	inline const FName& ImgMediaSmartCacheEnabled() 
+	inline const FName& ImgMediaSmartCacheEnabled()
 	{ static FName OptName(TEXT("ImgMediaSmartCacheEnabled")); return OptName; }
-	inline const FName& ImgMediaSmartCacheTimeToLookAhead() 
+	inline const FName& ImgMediaSmartCacheTimeToLookAhead()
 	{ static FName OptName(TEXT("ImgMediaSmartCacheTimeToLookAhead")); return OptName; }
 
 	/**
-	 * Requests parsing timecode meta data if available. 
+	 * Requests parsing timecode meta data if available.
 	 */
 	inline const FName& ParseTimecodeInfo()
 	{ static FName OptName(TEXT("ParseTimecodeInfo")); return OptName; }
