ElectraUtil: Added check for empty or fragmented mp4 file to avoid accessing non-existant structures

#rnx

diff --git a/Engine/Plugins/Media/ElectraUtil/Source/ElectraBase/Private/Utilities/MP4Boxes/MP4Track.cpp b/Engine/Plugins/Media/ElectraUtil/Source/ElectraBase/Private/Utilities/MP4Boxes/MP4Track.cpp
--- a/Engine/Plugins/Media/ElectraUtil/Source/ElectraBase/Private/Utilities/MP4Boxes/MP4Track.cpp
+++ b/Engine/Plugins/Media/ElectraUtil/Source/ElectraBase/Private/Utilities/MP4Boxes/MP4Track.cpp
@@ -184,7 +184,14 @@ namespace Electra
 			// Given the timeline mapping, locate the sample number that falls onto the start of the timeline
 			// and last one falling onto the end of the timeline.
 			TSharedPtr<FIterator, ESPMode::ThreadSafe> StartIt = CreateIterator();
-			check(StartIt.IsValid());
+
+			// If we could not create an iterator then this file is most likely a fragmented mp4
+			// which is not handled.
+			if (!StartIt.IsValid())
+			{
+				LastErrorMessage = FString::Printf(TEXT("Could not locate first media sample. Is this an empty or a fragmented mp4?"));
+				return false;
+			}
 			const int64 pts0 = Convs.CompositionTimeAtZeroPoint;
 			uint32 SyncSampleNum = 0;
 			bool bFound = false;
