Electra: Fixed PVS warnings

#rnx

diff --git a/Engine/Plugins/Media/ElectraPlayer/Source/ElectraPlayerRuntime/Private/Runtime/Utilities/UtilsMPEG.cpp b/Engine/Plugins/Media/ElectraPlayer/Source/ElectraPlayerRuntime/Private/Runtime/Utilities/UtilsMPEG.cpp
--- a/Engine/Plugins/Media/ElectraPlayer/Source/ElectraPlayerRuntime/Private/Runtime/Utilities/UtilsMPEG.cpp
+++ b/Engine/Plugins/Media/ElectraPlayer/Source/ElectraPlayerRuntime/Private/Runtime/Utilities/UtilsMPEG.cpp
@@ -170,9 +170,9 @@ namespace Electra
 				return false;
 			}
 			// "Unsynchronization" is not currently supported.
-			check((InData[5] & 0x80) == 0);
 			if ((InData[5] & 0x80) != 0)
 			{
+				check(!"unsynchronization not currently supported");
 				return false;
 			}
 			// Experimental headers are not supported
@@ -275,14 +275,12 @@ namespace Electra
 				{
 					uint32 FrameID = (static_cast<uint32>(InData[0]) << 24) | (static_cast<uint32>(InData[1]) << 16) | (static_cast<uint32>(InData[2]) << 8) | static_cast<uint32>(InData[3]);
 					int32 FrameSize = (int32)GetTagSize(InData + 4);
-					check(FrameSize > 10);
 					if (FrameSize <= 10)
 					{
 						return false;
 					}
 					const uint8* const FrameEnd = InData + FrameSize;
 					uint16 Flags = (InData[8] << 8) | InData[9];
-					check((Flags & 0xc0) == 0);
 					// Ignore compressed, encrypted or grouped tags.
 					if ((Flags & 0xc0) != 0)
 					{
diff --git a/Engine/Plugins/Media/ElectraUtil/Source/ElectraBase/Private/Core/MediaThreads.cpp b/Engine/Plugins/Media/ElectraUtil/Source/ElectraBase/Private/Core/MediaThreads.cpp
--- a/Engine/Plugins/Media/ElectraUtil/Source/ElectraBase/Private/Core/MediaThreads.cpp
+++ b/Engine/Plugins/Media/ElectraUtil/Source/ElectraBase/Private/Core/MediaThreads.cpp
@@ -37,7 +37,7 @@ void FMediaRunnable::Startup()
 void FMediaRunnable::Shutdown()
 {
 	check(MediaRunnablePrivate::AsyncExecutionThread);
-	
+
 	TFunction<void()> FinishTask = []()
 	{
 		MediaRunnablePrivate::bEndForShutdown = true;
@@ -70,16 +70,17 @@ void FMediaRunnable::EnqueueAsyncTask(TFunction<void()>&& InFunctionToExecuteOnA
  */
 FMediaRunnable* FMediaRunnable::Create(int32 InCoreAffinityMask, EThreadPriority InPriority, uint32 InStackSize, const FString& InThreadName)
 {
+	uint64 AffinityMask = (uint32)InCoreAffinityMask;
 	if (InCoreAffinityMask == 0 || InCoreAffinityMask == -1)
 	{
-		InCoreAffinityMask = FPlatformAffinity::GetNoAffinityMask();
+		AffinityMask = FPlatformAffinity::GetNoAffinityMask();
 	}
 
 	FMediaRunnable* t = new FMediaRunnable;
 	if (t)
 	{
 		t->ThreadPriority = InPriority;
-		t->InitialCoreAffinity = InCoreAffinityMask;
+		t->InitialCoreAffinity = AffinityMask;
 		t->StackSize = Align((uint32)InStackSize, 16384);
 		t->ThreadName = InThreadName;
 		t->bIsStarted = false;
