ElectraUtil: Validate that an mp4 edit list does not map non-existent media onto the timeline and if so, ignore the empty period and emit a log message

#rnx

diff --git a/Engine/Plugins/Media/ElectraUtil/Source/ElectraBase/Private/Utilities/MP4Boxes/MP4Track.cpp b/Engine/Plugins/Media/ElectraUtil/Source/ElectraBase/Private/Utilities/MP4Boxes/MP4Track.cpp
--- a/Engine/Plugins/Media/ElectraUtil/Source/ElectraBase/Private/Utilities/MP4Boxes/MP4Track.cpp
+++ b/Engine/Plugins/Media/ElectraUtil/Source/ElectraBase/Private/Utilities/MP4Boxes/MP4Track.cpp
@@ -137,7 +137,22 @@ namespace Electra
 					return false;
 				}
 
-				Convs.CompositionTimeAtZeroPoint = ElstEntries[0].media_time;
+				// The `media_time` in the edit list entry is specified in composition time. Typically an entry is used to shift
+				// a track having non-zero composition times to zero. In that case the `media_time` should correspond to the
+				// first composition time offset, or be greater.
+				// If it is less then the entry maps non-existing media, meaning that it essentially inserts an empty edit,
+				// which we do not support (what should be displayed then?)
+				if (ElstEntries[0].media_time < Convs.CompositionTimeAtZeroPoint)
+				{
+					// We assume that this is not actually wanted but a problem with the tool creating the file,
+					// so we emit a warning and ignore the edits `media_time`
+					UE_LOG(LogElectraBase, Verbose, TEXT("Edit list entry of track #%u maps non-existent media at composition time %lld to the timeline track. First available media composition time is %lld. Empty media will be ignored."), Convs.TrackID, (long long int)ElstEntries[0].media_time, (long long int)Convs.CompositionTimeAtZeroPoint);
+				}
+				else
+				{
+					Convs.CompositionTimeAtZeroPoint = ElstEntries[0].media_time;
+				}
+
 				// The mapped duration may be different from the media duration itself, in which case the track is either truncated
 				// or the last sample has a longer duration to be repeated until the end.
 				if (ElstEntries[0].edit_duration == 0)
