Reworked ElectraSamples, ElectraUtils and the decoder output of Electra Player

diff --git a/Engine/Plugins/Media/ElectraPlayer/Source/ElectraPlayerRuntime/Private/Runtime/Utilities/UtilsMP4.cpp b/Engine/Plugins/Media/ElectraPlayer/Source/ElectraPlayerRuntime/Private/Runtime/Utilities/UtilsMP4.cpp
--- a/Engine/Plugins/Media/ElectraPlayer/Source/ElectraPlayerRuntime/Private/Runtime/Utilities/UtilsMP4.cpp
+++ b/Engine/Plugins/Media/ElectraPlayer/Source/ElectraPlayerRuntime/Private/Runtime/Utilities/UtilsMP4.cpp
@@ -24,14 +24,14 @@ UtilsMP4::FMetadataParser::FMetadataParser()
 {
 	NumTotalItems = 0;
 
-	WellKnownItems.Emplace(ElectraDecodersUtil::MakeMP4Atom(0xa9,'n','a','m'), FMetadataTags::Title);
-	WellKnownItems.Emplace(ElectraDecodersUtil::MakeMP4Atom(0xa9,'A','R','T'), FMetadataTags::Artist);
-	WellKnownItems.Emplace(ElectraDecodersUtil::MakeMP4Atom(0xa9,'g','e','n'), FMetadataTags::Genre);
-	WellKnownItems.Emplace(ElectraDecodersUtil::MakeMP4Atom(0xa9,'d','a','y'), FMetadataTags::Date);
-	WellKnownItems.Emplace(ElectraDecodersUtil::MakeMP4Atom('d','e','s','c'), FMetadataTags::Description);
-	WellKnownItems.Emplace(ElectraDecodersUtil::MakeMP4Atom('l','d','e','s'), FMetadataTags::LongDescription);
-	WellKnownItems.Emplace(ElectraDecodersUtil::MakeMP4Atom(0xa9,'a','l','b'), FMetadataTags::Album);
-	WellKnownItems.Emplace(ElectraDecodersUtil::MakeMP4Atom(0xa9,'t','o','o'), FMetadataTags::Encoder);
+	WellKnownItems.Emplace(Electra::UtilitiesMP4::MakeBoxAtom(0xa9u,'n','a','m'), FMetadataTags::Title);
+	WellKnownItems.Emplace(Electra::UtilitiesMP4::MakeBoxAtom(0xa9u,'A','R','T'), FMetadataTags::Artist);
+	WellKnownItems.Emplace(Electra::UtilitiesMP4::MakeBoxAtom(0xa9u,'g','e','n'), FMetadataTags::Genre);
+	WellKnownItems.Emplace(Electra::UtilitiesMP4::MakeBoxAtom(0xa9u,'d','a','y'), FMetadataTags::Date);
+	WellKnownItems.Emplace(Electra::UtilitiesMP4::MakeBoxAtom('d','e','s','c'), FMetadataTags::Description);
+	WellKnownItems.Emplace(Electra::UtilitiesMP4::MakeBoxAtom('l','d','e','s'), FMetadataTags::LongDescription);
+	WellKnownItems.Emplace(Electra::UtilitiesMP4::MakeBoxAtom(0xa9u,'a','l','b'), FMetadataTags::Album);
+	WellKnownItems.Emplace(Electra::UtilitiesMP4::MakeBoxAtom(0xa9u,'t','o','o'), FMetadataTags::Encoder);
 }
 
 FString UtilsMP4::FMetadataParser::PrintableBoxAtom(const uint32 InAtom)
@@ -90,7 +90,7 @@ void UtilsMP4::FMetadataParser::AddItem(const FString& InType, const FString& In
 UtilsMP4::FMetadataParser::EResult UtilsMP4::FMetadataParser::Parse(uint32 InHandler, uint32 InHandlerReserved0, const TArray<UtilsMP4::FMetadataParser::FBoxInfo>& InBoxes)
 {
 	// We only support the Apple iTunes metadata at the moment.
-	if (InHandler != ElectraDecodersUtil::MakeMP4Atom('m','d','i','r'))
+	if (InHandler != Electra::UtilitiesMP4::MakeBoxAtom('m','d','i','r'))
 	{
 		return UtilsMP4::FMetadataParser::EResult::NotSupported;
 	}
@@ -122,7 +122,7 @@ UtilsMP4::FMetadataParser::EResult UtilsMP4::FMetadataParser::Parse(uint32 InHan
 	bool bHaveILST = false;
 	for(int32 i=0; i<InBoxes.Num(); ++i)
 	{
-		if (InBoxes[i].Type == ElectraDecodersUtil::MakeMP4Atom('i','l','s','t'))
+		if (InBoxes[i].Type == Electra::UtilitiesMP4::MakeBoxAtom('i','l','s','t'))
 		{
 			// We will actually parse all 'ilst' boxes if for some reason the metadata has
 			// been scattered across several.
@@ -140,7 +140,7 @@ void UtilsMP4::FMetadataParser::Parse(const UtilsMP4::FMetadataParser::FBoxInfo&
 		return;
 	}
 
-	ElectraDecodersUtil::FMP4AtomReader dr(InBox.Data, InBox.Size);
+	Electra::UtilitiesMP4::FMP4AtomReader dr(MakeConstArrayView<const uint8>(reinterpret_cast<const uint8*>(InBox.Data), InBox.Size));
 	while(dr.GetNumBytesRemaining() > 8)
 	{
 		int32 BoxSize;
@@ -153,7 +153,7 @@ void UtilsMP4::FMetadataParser::Parse(const UtilsMP4::FMetadataParser::FBoxInfo&
 		}
 		else switch(BoxType)
 		{
-			case ElectraDecodersUtil::MakeMP4Atom('-','-','-','-'):
+			case Electra::UtilitiesMP4::MakeBoxAtom('-','-','-','-'):
 			{
 				ParseBoxDataiTunes(dr.GetCurrentDataPointer(), BoxSize-8);
 				break;
@@ -176,7 +176,7 @@ void UtilsMP4::FMetadataParser::ParseBoxDataList(const FString& AsCategory, cons
 		return;
 	}
 
-	ElectraDecodersUtil::FMP4AtomReader dr(InBoxData, InBoxSize);
+	Electra::UtilitiesMP4::FMP4AtomReader dr(MakeConstArrayView<const uint8>(reinterpret_cast<const uint8*>(InBoxData), InBoxSize));
 
 #define RETURN_IF_READERROR(expr)	\
 	if (!(expr))					\
@@ -190,7 +190,7 @@ void UtilsMP4::FMetadataParser::ParseBoxDataList(const FString& AsCategory, cons
 		RETURN_IF_READERROR(dr.Read(BoxSize));
 		RETURN_IF_READERROR(dr.Read(BoxType));
 		CurrentOffset = dr.GetCurrentOffset();
-		if (BoxType != ElectraDecodersUtil::MakeMP4Atom('d','a','t','a'))
+		if (BoxType != Electra::UtilitiesMP4::MakeBoxAtom('d','a','t','a'))
 		{
 			continue;
 		}
@@ -353,7 +353,7 @@ void UtilsMP4::FMetadataParser::ParseBoxDataiTunes(const uint8* InBoxData, uint3
 		return;
 	}
 
-	ElectraDecodersUtil::FMP4AtomReader dr(InBoxData, InBoxSize);
+	Electra::UtilitiesMP4::FMP4AtomReader dr(MakeConstArrayView<const uint8>(reinterpret_cast<const uint8*>(InBoxData), InBoxSize));
 
 #define RETURN_IF_READERROR(expr)	\
 	if (!(expr))					\
@@ -368,7 +368,7 @@ void UtilsMP4::FMetadataParser::ParseBoxDataiTunes(const uint8* InBoxData, uint3
 		RETURN_IF_READERROR(dr.Read(BoxSize));
 		RETURN_IF_READERROR(dr.Read(BoxType));
 		CurrentOffset = dr.GetCurrentOffset();
-		if (BoxType == ElectraDecodersUtil::MakeMP4Atom('m','e','a','n'))
+		if (BoxType == Electra::UtilitiesMP4::MakeBoxAtom('m','e','a','n'))
 		{
 			// 4 reserved bytes follow
 			RETURN_IF_READERROR(dr.ReadBytes(nullptr, 4));
@@ -380,14 +380,14 @@ void UtilsMP4::FMetadataParser::ParseBoxDataiTunes(const uint8* InBoxData, uint3
 				return;
 			}
 		}
-		else if (BoxType == ElectraDecodersUtil::MakeMP4Atom('n','a','m','e'))
+		else if (BoxType == Electra::UtilitiesMP4::MakeBoxAtom('n','a','m','e'))
 		{
 			// 4 reserved bytes follow
 			RETURN_IF_READERROR(dr.ReadBytes(nullptr, 4));
 			int32 NumRemainingDataBytes = BoxSize - 8 - (dr.GetCurrentOffset() - CurrentOffset);
 			RETURN_IF_READERROR(dr.ReadStringUTF8(Name, NumRemainingDataBytes));
 		}
-		else if (BoxType == ElectraDecodersUtil::MakeMP4Atom('d','a','t','a'))
+		else if (BoxType == Electra::UtilitiesMP4::MakeBoxAtom('d','a','t','a'))
 		{
 			ParseBoxDataList(FString(TEXT("iTunes."))+Name, dr.GetCurrentDataPointer()-8, BoxSize);
 		}
@@ -685,7 +685,7 @@ bool UtilsMP4::FMP4RootBoxLocator::LocateRootBoxes(TArray<FBoxInfo>& OutBoxInfos
 				BoxInternalOffset += 8;
 			}
 
-			if (bi.Type == Make4CC('u','u','i','d'))
+			if (bi.Type == UtilitiesMP4::MakeBoxAtom('u','u','i','d'))
 			{
 				if (DataBuffer->Num() < BoxInternalOffset + 16)
 				{
diff --git a/Engine/Plugins/Media/ElectraPlayer/Source/ElectraPlayerRuntime/Private/Runtime/Utilities/UtilsMP4.h b/Engine/Plugins/Media/ElectraPlayer/Source/ElectraPlayerRuntime/Private/Runtime/Utilities/UtilsMP4.h
--- a/Engine/Plugins/Media/ElectraPlayer/Source/ElectraPlayerRuntime/Private/Runtime/Utilities/UtilsMP4.h
+++ b/Engine/Plugins/Media/ElectraPlayer/Source/ElectraPlayerRuntime/Private/Runtime/Utilities/UtilsMP4.h
@@ -5,7 +5,7 @@
 #include <CoreMinimal.h>
 #include <Misc/Variant.h>
 #include "HTTP/HTTPManager.h"
-#include "Utils/MPEG/ElectraUtilsMP4.h"
+#include "Utilities/UtilitiesMP4.h"
 #include "MediaStreamMetadata.h"
 
 namespace Electra
@@ -14,12 +14,6 @@ namespace Electra
 class UtilsMP4
 {
 public:
-	static constexpr uint32 Make4CC(const uint8 A, const uint8 B, const uint8 C, const uint8 D)
-	{
-		return (static_cast<uint32>(A) << 24) | (static_cast<uint32>(B) << 16) | (static_cast<uint32>(C) << 8) | static_cast<uint32>(D);
-	}
-
-
 	/**
 	 * This class parses metadata embedded in an mp4 / ISO14496-12 file.
 	 * Presently only the structure as used and defined by Apple iTunes is supported.
diff --git a/Engine/Plugins/Media/ElectraPlayer/Source/ElectraProtron/ElectraProtron.Build.cs b/Engine/Plugins/Media/ElectraPlayer/Source/ElectraProtron/ElectraProtron.Build.cs
--- a/Engine/Plugins/Media/ElectraPlayer/Source/ElectraProtron/ElectraProtron.Build.cs
+++ b/Engine/Plugins/Media/ElectraPlayer/Source/ElectraProtron/ElectraProtron.Build.cs
@@ -21,9 +21,7 @@ namespace UnrealBuildTool.Rules
 					"RHI",
 					"ElectraBase",
                     "ElectraDecoders",
-                    "ElectraPlayerPlugin",      // for early start and init of what we want to access
-                    "ElectraPlayerRuntime",     // sharing the resource manager
-					"ElectraSamples",           // and renderer
+					"ElectraSamples"
 				});
 
 			PrivateIncludePathModuleNames.AddRange(
diff --git a/Engine/Plugins/Media/ElectraPlayer/Source/ElectraProtron/Private/ElectraProtronModule.cpp b/Engine/Plugins/Media/ElectraPlayer/Source/ElectraProtron/Private/ElectraProtronModule.cpp
--- a/Engine/Plugins/Media/ElectraPlayer/Source/ElectraProtron/Private/ElectraProtronModule.cpp
+++ b/Engine/Plugins/Media/ElectraPlayer/Source/ElectraProtron/Private/ElectraProtronModule.cpp
@@ -13,24 +13,23 @@ class FElectraProtronModule
 	: public IElectraProtronModule
 {
 public:
-
 	//~ IElectraProtronModule interface
-
-	virtual TSharedPtr<IMediaPlayer, ESPMode::ThreadSafe> CreatePlayer(IMediaEventSink& EventSink) override
+	TSharedPtr<IMediaPlayer, ESPMode::ThreadSafe> CreatePlayer(IMediaEventSink& EventSink) override
 	{
 		return MakeShared<FElectraProtronPlayer, ESPMode::ThreadSafe>(EventSink);
 	}
 
-public:
-	virtual void StartupModule() override
+	void StartupModule() override
+	{
+		// Load the modules we depend on. They may have been loaded already, but we do it explicitly here to ensure that
+		// they will not be unloaded on shutdown before this module here.
+		FModuleManager::Get().LoadModule(TEXT("ElectraBase"));
+		FModuleManager::Get().LoadModule(TEXT("ElectraSamples"));
+		FModuleManager::Get().LoadModule(TEXT("ElectraCodecFactory"));
+	}
+	void ShutdownModule() override
 	{
-		// We use resource functionality from the Electra Player, so we make sure the modules are loaded.
-		FModuleManager::Get().GetModule("ElectraPlayerPlugin");
-		FModuleManager::Get().GetModule("ElectraPlayerRuntime");
 	}
-	virtual void ShutdownModule() override
-	{ }
-
 };
 
 IMPLEMENT_MODULE(FElectraProtronModule, ElectraProtron)
diff --git a/Engine/Plugins/Media/ElectraPlayer/Source/ElectraProtron/Private/Player/ElectraProtronPlayerDecoder.cpp b/Engine/Plugins/Media/ElectraPlayer/Source/ElectraProtron/Private/Player/ElectraProtronPlayerDecoder.cpp
--- a/Engine/Plugins/Media/ElectraPlayer/Source/ElectraProtron/Private/Player/ElectraProtronPlayerDecoder.cpp
+++ b/Engine/Plugins/Media/ElectraPlayer/Source/ElectraProtron/Private/Player/ElectraProtronPlayerDecoder.cpp
@@ -11,38 +11,17 @@
 #include "Utils/MPEG/ElectraUtilsMPEGAudio.h"
 #include "TrackFormatInfo.h"
 
-#include "IElectraPlayerDecoderResourceManager.h"
 #include "IElectraDecoder.h"
 #include "IElectraDecoderOutputVideo.h"
 #include "IElectraDecoderOutputAudio.h"
-#include "IElectraPlayerRuntimeModule.h"
 #include "ElectraDecodersUtils.h"
 
 #include "MediaSamples.h"
-#include "MediaRendererBase.h"
 #include "MediaDecoderOutput.h"
 #include "ElectraTextureSample.h"
 #include "IElectraAudioSample.h"
 
 
-namespace ProtronDecoderHelpers
-{
-	class FDecoderOutputPointerWrapper : public Electra::IMediaRenderer::IBuffer
-	{
-	public:
-		FDecoderOutputPointerWrapper(FVideoDecoderOutputPtr InPtr) : Ptr(MoveTemp(InPtr))
-		{ HelperDict.Set(Electra::RenderOptionKeys::Texture, Electra::FVariantValue(Ptr)); }
-		const Electra::FParamDict& GetBufferProperties() const override
-		{ return HelperDict; }
-		Electra::FParamDict& GetMutableBufferProperties() override
-		{ return HelperDict; }
-	private:
-		Electra::FParamDict HelperDict;
-		FVideoDecoderOutputPtr Ptr;
-	};
-}
-
-
 
 void FElectraProtronPlayer::FImpl::FDecoderThread::StartThread(const FOpenParam& InParam, const TSharedPtr<FSharedPlayParams, ESPMode::ThreadSafe>& InSharedPlayParams)
 {
@@ -952,13 +931,11 @@ void FElectraProtronPlayer::FImpl::FDecoderThread::HandleOutputFrame()
 				// Send the output along?
 				if (bSendToOutput)
 				{
-					TSharedPtr<IElectraDecoderVideoOutput, ESPMode::ThreadSafe> VideoDecoderOutput = StaticCastSharedPtr<IElectraDecoderVideoOutput>(CurrentDecoderOutput);
-					FVideoDecoderOutputPtr NewDecoderOutput = VideoDecoderOutputPool->AcquireShared();
-					// Note: If we need to be informed when the decoder output is no longer needed we
-					//       could set ourselves or some notification wrapper as the owner (derived from IDecoderOutputOwner)
-					//       to receive a SampleReleasedToPool() callback.
-					//NewDecoderOutput->SetOwner(SharedThis(this));
-					ProtronDecoderHelpers::FDecoderOutputPointerWrapper DecOutputPtrWrapper(NewDecoderOutput);
+					auto VideoDecoderOutput = StaticCastSharedPtr<IElectraDecoderVideoOutput>(CurrentDecoderOutput);
+
+					FElectraTextureSamplePtr OutputTextureSample = VideoDecoderOutputPool->AcquireShared();
+					// Note: If we need to be informed when the output sample is no longer needed we
+					//       could bind a delegate to it via GetReleaseDelegate()
 					TSharedPtr<Electra::FParamDict, ESPMode::ThreadSafe> BufferProperties = MakeShared<Electra::FParamDict, ESPMode::ThreadSafe>();
 
 					Electra::FTimeValue OutputPTS, OutputDur;
@@ -987,11 +964,9 @@ void FElectraProtronPlayer::FImpl::FDecoderThread::HandleOutputFrame()
 					}
 
 					FString DecoderOutputErrorMsg;
-					if (ensure(FElectraPlayerDecoderResourceManager::SetupRenderBufferFromDecoderOutput(DecoderOutputErrorMsg, &DecOutputPtrWrapper, BufferProperties, VideoDecoderOutput, PlatformResource)))
+					if (ensure(FElectraDecodersPlatformResources::SetupOutputTextureSample(DecoderOutputErrorMsg, OutputTextureSample, VideoDecoderOutput, BufferProperties, DecoderPlatformResource)))
 					{
-						FElectraTextureSampleRef TextureSample = Params.TexturePool->AcquireShared();
-						TextureSample->Initialize(NewDecoderOutput.Get());
-						Params.SampleQueueInterface->EnqueueVideoSample(TextureSample, RawPTS, RawDur);
+						Params.SampleQueueInterface->EnqueueVideoSample(OutputTextureSample.ToSharedRef(), RawPTS, RawDur);
 					}
 				}
 				CurrentDecoderOutput.Reset();
@@ -1074,11 +1049,12 @@ void FElectraProtronPlayer::FImpl::FDecoderThread::HandleOutputFrame()
 							// Anything remaining to pass along?
 							if (NumSamplesProduced)
 							{
-								// Pre-initialize the sample without copying data.
-								if (AudioSample->Initialize(nullptr, EMediaAudioSampleFormat::Float, (uint32)NumberOfChannels, (uint32)NumSamplesProduced, (uint32)SamplingRate, PTS, Dur))
+								// Allocate the sample buffer.
+								if (AudioSample->AllocateFor(EMediaAudioSampleFormat::Float, (uint32)NumberOfChannels, (uint32)NumSamplesProduced))
 								{
-
-									AudioChannelMapper.MapChannels(const_cast<void*>(AudioSample->GetBuffer()), (int32)AudioSample->GetAllocatedSize(),
+									// Set the sample parameters and copy the data.
+									AudioSample->SetParameters((uint32)SamplingRate, PTS, Dur);
+									AudioChannelMapper.MapChannels(AudioSample->GetWritableBuffer(), (int32)AudioSample->GetAllocatedSize(),
 																	ElectraDecodersUtil::AdvancePointer(PCMBuffer, ByteOffsetToFirstSample), NumSamplesProduced * NumBytesPerFrame, NumSamplesProduced);
 									Params.SampleQueueInterface->EnqueueAudioSample(AudioSample);
 								}
@@ -1127,10 +1103,10 @@ void FElectraProtronPlayer::FImpl::FDecoderThread::DestroyDecoder()
 		DecoderBitstreamProcessor->Clear();
 		DecoderBitstreamProcessor.Reset();
 	}
-	if (PlatformResource)
+	if (DecoderPlatformResource)
 	{
-		FElectraPlayerDecoderResourceManager::GetDelegate()->ReleasePlatformResource(this, PlatformResource);
-		PlatformResource = nullptr;
+		FElectraDecodersPlatformResources::ReleasePlatformVideoResource(this, DecoderPlatformResource);
+		DecoderPlatformResource = nullptr;
 	}
 }
 
@@ -1151,9 +1127,6 @@ bool FElectraProtronPlayer::FImpl::FDecoderThread::CreateDecoder()
 	// Get the decoder factory module.
 	IElectraCodecFactoryModule* FactoryModule = static_cast<IElectraCodecFactoryModule*>(FModuleManager::Get().GetModule(TEXT("ElectraCodecFactory")));
 	check(FactoryModule);
-	// And since we piggyback on the Electra Player decoder output handling and conversion we need the plugin as well.
-	IElectraPlayerRuntimeModule* PluginModule = static_cast<IElectraPlayerRuntimeModule*>(FModuleManager::Get().GetModule(TEXT("ElectraPlayerRuntime")));
-	check(PluginModule && PluginModule->IsInitialized());
 
 	CurrentCodecSpecificData.Empty();
 	CurrentColorimetry.Reset();
@@ -1230,29 +1203,27 @@ bool FElectraProtronPlayer::FImpl::FDecoderThread::CreateDecoder()
 		return false;
 	}
 
-	// Create a decoder instance, for which we need to get delegates from the Electra Player.
-	TSharedPtr<IElectraDecoderResourceDelegate, ESPMode::ThreadSafe> ElectraDecoderResourceDelegate = FElectraPlayerDecoderResourceManager::GetDelegate();
+
 	// Create platform specific resources to be used with the new decoder.
-	check(PlatformResource == nullptr);
+	check(DecoderPlatformResource == nullptr);
 	if (DecoderTypeIndex == CodecTypeIndex(ElectraProtronUtils::FCodecInfo::EType::Video))
 	{
-		PlatformResource = ElectraDecoderResourceDelegate->CreatePlatformResource(this, IElectraDecoderResourceDelegateBase::EDecoderPlatformResourceType::Video, DecoderCfgParams);
+		DecoderPlatformResource = FElectraDecodersPlatformResources::CreatePlatformVideoResource(this, DecoderCfgParams);
 	}
-	//DecoderCreateCfg.Emplace(TEXT("platformresource"), FVariant((uint64)PlatformResource));
 
-	// Try to create the decoder now.
+	// Create a decoder instance.
 	DecoderInstance = DecoderFactory->CreateDecoderForFormat(DecoderFormat, DecoderCfgParams);
 	if (!DecoderInstance.IsValid())
 	{
-		ElectraDecoderResourceDelegate->ReleasePlatformResource(this, PlatformResource);
-		PlatformResource = nullptr;
+		FElectraDecodersPlatformResources::ReleasePlatformVideoResource(this, DecoderPlatformResource);
+		DecoderPlatformResource = nullptr;
 		LastErrorMessage = FString::Printf(TEXT("Failed to create decoder for format \"%s\"."), *DecoderFormat);
 		return false;
 	}
 	if (DecoderInstance->GetError().IsSet())
 	{
-		ElectraDecoderResourceDelegate->ReleasePlatformResource(this, PlatformResource);
-		PlatformResource = nullptr;
+		FElectraDecodersPlatformResources::ReleasePlatformVideoResource(this, DecoderPlatformResource);
+		DecoderPlatformResource = nullptr;
 		LastErrorMessage = DecoderInstance->GetError().GetMessage();
 		return false;
 	}
diff --git a/Engine/Plugins/Media/ElectraPlayer/Source/ElectraProtron/Private/Player/ElectraProtronPlayerImpl.cpp b/Engine/Plugins/Media/ElectraPlayer/Source/ElectraProtron/Private/Player/ElectraProtronPlayerImpl.cpp
--- a/Engine/Plugins/Media/ElectraPlayer/Source/ElectraProtron/Private/Player/ElectraProtronPlayerImpl.cpp
+++ b/Engine/Plugins/Media/ElectraPlayer/Source/ElectraProtron/Private/Player/ElectraProtronPlayerImpl.cpp
@@ -538,7 +538,7 @@ void FElectraProtronPlayer::FImpl::InternalOpen(const FString& InFilename)
 			{
 				auto GetValue = [](const TSharedPtr<Electra::UtilitiesMP4::FMP4BoxBase, ESPMode::ThreadSafe>& InBox) -> FString
 				{
-					Electra::UtilitiesMP4::FMP4AtomReaderHelper timReader(InBox->GetBoxData());
+					Electra::UtilitiesMP4::FMP4AtomReader timReader(InBox->GetBoxData());
 					FString String;
 					uint16 StringLength, UnknownValue;
 					if (timReader.Read(StringLength)&&  timReader.Read(UnknownValue))
diff --git a/Engine/Plugins/Media/ElectraPlayer/Source/ElectraProtron/Private/Player/ElectraProtronPlayerImpl.h b/Engine/Plugins/Media/ElectraPlayer/Source/ElectraProtron/Private/Player/ElectraProtronPlayerImpl.h
--- a/Engine/Plugins/Media/ElectraPlayer/Source/ElectraProtron/Private/Player/ElectraProtronPlayerImpl.h
+++ b/Engine/Plugins/Media/ElectraPlayer/Source/ElectraProtron/Private/Player/ElectraProtronPlayerImpl.h
@@ -16,11 +16,12 @@
 #include "Utilities/MP4Boxes/MP4Track.h"
 #include "TrackFormatInfo.h"
 #include "IElectraDecoder.h"
-#include "IElectraPlayerDecoderResourceManager.h"
+#include "ElectraDecodersPlatformResources.h"
+#include "MediaSamples.h"
+#include "ElectraTextureSample.h"
 #include "Utils/MPEG/ElectraUtilsMPEGVideo.h"
 #include "Utils/AudioChannelMapper.h"
-#include "Decoder/VideoDecoderHelpers.h"
-#include "MediaSamples.h"
+#include "Utils/VideoDecoderHelpers.h"
 #include "ElectraProtronPlayerCache.h"
 
 
@@ -553,7 +554,7 @@ private:
 		FDecoderThread(const FConfig& InConfig, int32 InCodecTypeIndex)
 			: Config(InConfig), DecoderTypeIndex(InCodecTypeIndex)
 		{
-			VideoDecoderOutputPool = FVideoPool::Create();
+			VideoDecoderOutputPool = MakeShareable(new FElectraTextureSamplePool);
 		}
 
 		// Inherited from FRunnable
@@ -688,7 +689,7 @@ private:
 		TSharedPtr<IElectraDecoderBitstreamProcessor, ESPMode::ThreadSafe> DecoderBitstreamProcessor;
 
 		TMap<FString, FVariant> CurrentCodecSpecificData;
-		IElectraDecoderResourceDelegateBase::IDecoderPlatformResource* PlatformResource = nullptr;
+		FElectraDecodersPlatformResources::IDecoderPlatformResource* DecoderPlatformResource = nullptr;
 		TUniquePtr<FInDecoder> CurrentInputSample;
 		TSharedPtr<IElectraDecoderOutput, ESPMode::ThreadSafe> CurrentDecoderOutput;
 		TUniquePtr<FInDecoder> InputForCurrentDecoderOutput;
@@ -703,8 +704,7 @@ private:
 		bool bWaitForSyncSample = true;
 		bool bWarnedMissingSyncSample = false;
 
-		using FVideoPool = TDecoderOutputObjectPool<FVideoDecoderOutput, FElectraPlayerDecoderResourceManager::FVideo>;
-		TSharedPtr<FVideoPool, ESPMode::ThreadSafe> VideoDecoderOutputPool;
+		TSharedPtr<FElectraTextureSamplePool, ESPMode::ThreadSafe> VideoDecoderOutputPool;
 	};
 
 	FLoaderThread VideoLoaderThread {Config, CodecTypeIndex(ElectraProtronUtils::FCodecInfo::EType::Video)};
diff --git a/Engine/Plugins/Media/ElectraUtil/Source/ElectraBase/Private/ElectraBaseModule.cpp b/Engine/Plugins/Media/ElectraUtil/Source/ElectraBase/Private/ElectraBaseModule.cpp
--- a/Engine/Plugins/Media/ElectraUtil/Source/ElectraBase/Private/ElectraBaseModule.cpp
+++ b/Engine/Plugins/Media/ElectraUtil/Source/ElectraBase/Private/ElectraBaseModule.cpp
@@ -2,10 +2,13 @@
 
 #include "ElectraBaseModule.h"
 
-#include "CoreMinimal.h"
+#include "CoreTypes.h"
 #include "Modules/ModuleManager.h"
 #include "IElectraBaseModule.h"
 
+#include "Core/MediaThreads.h"
+
+
 #define LOCTEXT_NAMESPACE "ElectraBaseModule"
 
 DEFINE_LOG_CATEGORY(LogElectraBase);
@@ -14,19 +17,28 @@ DEFINE_LOG_CATEGORY(LogElectraBase);
 
 class FElectraBaseModule: public IElectraBaseModule
 {
-public:
 public:
 	// IModuleInterface interface
 
 	void StartupModule() override
 	{
+		if (!bInitialized)
+		{
+			bInitialized = true;
+			FMediaRunnable::Startup();
+		}
 	}
 
 	void ShutdownModule() override
 	{
+		if (bInitialized)
+		{
+			FMediaRunnable::Shutdown();
+		}
 	}
 
 private:
+	bool bInitialized = false;
 };
 
 IMPLEMENT_MODULE(FElectraBaseModule, ElectraBase);
diff --git a/Engine/Plugins/Media/ElectraUtil/Source/ElectraBase/Private/Utilities/MP4Boxes/MP4Boxes.cpp b/Engine/Plugins/Media/ElectraUtil/Source/ElectraBase/Private/Utilities/MP4Boxes/MP4Boxes.cpp
--- a/Engine/Plugins/Media/ElectraUtil/Source/ElectraBase/Private/Utilities/MP4Boxes/MP4Boxes.cpp
+++ b/Engine/Plugins/Media/ElectraUtil/Source/ElectraBase/Private/Utilities/MP4Boxes/MP4Boxes.cpp
@@ -11,7 +11,7 @@ namespace Electra
 	namespace UtilitiesMP4
 	{
 
-		void FMP4BoxBase::ProcessBoxChildrenRecursively(FMP4AtomReaderHelper& InOutReader, const FMP4BoxInfo& InCurrentBoxInfo)
+		void FMP4BoxBase::ProcessBoxChildrenRecursively(FMP4AtomReader& InOutReader, const FMP4BoxInfo& InCurrentBoxInfo)
 		{
 			FMP4BoxInfo bi;
 			while(InOutReader.ParseIntoBoxInfo(bi, InCurrentBoxInfo.Offset + InCurrentBoxInfo.DataOffset + InOutReader.GetCurrentOffset()))
@@ -19,7 +19,7 @@ namespace Electra
 				TSharedPtr<FMP4BoxBase, ESPMode::ThreadSafe> Child = FMP4BoxFactory::Get().Create(AsWeak(), bi);
 				if (AddChildBox(Child) && !Child->IsLeafBox() && !Child->IsListOfEntries())
 				{
-					FMP4AtomReaderHelper ar(Child->BoxInfo.Data);
+					FMP4AtomReader ar(Child->BoxInfo.Data);
 					Child->ProcessBoxChildrenRecursively(ar, Child->BoxInfo);
 				}
 				InOutReader.SkipBytes(bi.Size - bi.DataOffset);
@@ -41,7 +41,7 @@ namespace Electra
 				return;
 			}
 			Parsed = MakePimpl<FParsed>();
-			FMP4AtomReaderHelper ar(BoxInfo.Data);
+			FMP4AtomReader ar(BoxInfo.Data);
 			ar.ReadVersionAndFlags(Parsed->Version, Parsed->Flags);
 			if (Parsed->Version == 1)
 			{
@@ -94,7 +94,7 @@ namespace Electra
 				return;
 			}
 			Parsed = MakePimpl<FParsed>();
-			FMP4AtomReaderHelper ar(BoxInfo.Data);
+			FMP4AtomReader ar(BoxInfo.Data);
 			ar.ReadVersionAndFlags(Parsed->Version, Parsed->Flags);
 			if (Parsed->Version == 1)
 			{
@@ -170,7 +170,7 @@ namespace Electra
 				return;
 			}
 			Parsed = MakePimpl<FParsed>();
-			FMP4AtomReaderHelper ar(BoxInfo.Data);
+			FMP4AtomReader ar(BoxInfo.Data);
 			FMP4BoxInfo bi;
 			while(ar.ParseIntoBoxInfo(bi, BoxInfo.Offset + BoxInfo.DataOffset + ar.GetCurrentOffset()))
 			{
@@ -219,7 +219,7 @@ namespace Electra
 				return;
 			}
 			Parsed = MakePimpl<FParsed>();
-			FMP4AtomReaderHelper ar(BoxInfo.Data);
+			FMP4AtomReader ar(BoxInfo.Data);
 			ar.ReadVersionAndFlags(Parsed->Version, Parsed->Flags);
 			uint32 NumEntries;
 			ar.Read(NumEntries);
@@ -280,7 +280,7 @@ namespace Electra
 				return;
 			}
 			Parsed = MakePimpl<FParsed>();
-			FMP4AtomReaderHelper ar(BoxInfo.Data);
+			FMP4AtomReader ar(BoxInfo.Data);
 			ar.ReadVersionAndFlags(Parsed->Version, Parsed->Flags);
 			if (Parsed->Version == 1)
 			{
@@ -337,7 +337,7 @@ namespace Electra
 				return;
 			}
 			Parsed = MakePimpl<FParsed>();
-			FMP4AtomReaderHelper ar(BoxInfo.Data);
+			FMP4AtomReader ar(BoxInfo.Data);
 			ar.ReadVersionAndFlags(Parsed->Version, Parsed->Flags);
 			ar.SkipBytes(4);				// `pre_defined`
 			ar.Read(Parsed->HandlerType);
@@ -371,7 +371,7 @@ namespace Electra
 				return;
 			}
 			Parsed = MakePimpl<FParsed>();
-			FMP4AtomReaderHelper ar(BoxInfo.Data);
+			FMP4AtomReader ar(BoxInfo.Data);
 			ar.ReadVersionAndFlags(Parsed->Version, Parsed->Flags);
 			FString Lang;
 			ar.ReadStringUTF8(Lang, ar.GetNumBytesRemaining());
@@ -399,7 +399,7 @@ namespace Electra
 				return;
 			}
 			Parsed = MakePimpl<FParsed>();
-			FMP4AtomReaderHelper ar(BoxInfo.Data);
+			FMP4AtomReader ar(BoxInfo.Data);
 			ar.ReadVersionAndFlags(Parsed->Version, Parsed->Flags);
 
 			// In order to parse the sample entry we need to know the media handler type.
@@ -541,7 +541,7 @@ namespace Electra
 				return;
 			}
 			Parsed = MakePimpl<FParsed>();
-			FMP4AtomReaderHelper ar(BoxInfo.Data);
+			FMP4AtomReader ar(BoxInfo.Data);
 			ar.ReadVersionAndFlags(Parsed->Version, Parsed->Flags);
 			uint32 NumEntries;
 			ar.Read(NumEntries);
@@ -573,7 +573,7 @@ namespace Electra
 				return;
 			}
 			Parsed = MakePimpl<FParsed>();
-			FMP4AtomReaderHelper ar(BoxInfo.Data);
+			FMP4AtomReader ar(BoxInfo.Data);
 			ar.ReadVersionAndFlags(Parsed->Version, Parsed->Flags);
 			uint32 NumEntries;
 			ar.Read(NumEntries);
@@ -604,7 +604,7 @@ namespace Electra
 				return;
 			}
 			Parsed = MakePimpl<FParsed>();
-			FMP4AtomReaderHelper ar(BoxInfo.Data);
+			FMP4AtomReader ar(BoxInfo.Data);
 			// Members of the generic `SampleEntry` class
 			ar.SkipBytes(6);
 			ar.Read(Parsed->DataReferenceIndex);
@@ -669,7 +669,7 @@ namespace Electra
 				return;
 			}
 			Parsed = MakePimpl<FParsed>();
-			FMP4AtomReaderHelper ar(BoxInfo.Data);
+			FMP4AtomReader ar(BoxInfo.Data);
 			// Members of the generic `SampleEntry` class
 			ar.SkipBytes(6);
 			ar.Read(Parsed->DataReferenceIndex);
@@ -804,7 +804,7 @@ namespace Electra
 				return;
 			}
 			Parsed = MakePimpl<FParsed>();
-			FMP4AtomReaderHelper ar(BoxInfo.Data);
+			FMP4AtomReader ar(BoxInfo.Data);
 			// Members of the generic `SampleEntry` class
 			ar.SkipBytes(6);
 			ar.Read(Parsed->DataReferenceIndex);
@@ -858,7 +858,7 @@ namespace Electra
 				return;
 			}
 			Parsed = MakePimpl<FParsed>();
-			FMP4AtomReaderHelper ar(BoxInfo.Data);
+			FMP4AtomReader ar(BoxInfo.Data);
 			// Members of the generic `SampleEntry` class
 			ar.SkipBytes(6);
 			ar.Read(Parsed->DataReferenceIndex);
@@ -894,7 +894,7 @@ namespace Electra
 				return;
 			}
 			Parsed = MakePimpl<FParsed>();
-			FMP4AtomReaderHelper ar(BoxInfo.Data);
+			FMP4AtomReader ar(BoxInfo.Data);
 			Parsed->AVCDecoderConfigurationRecord = MakeConstArrayView(ar.GetCurrentDataPointer(), ar.GetNumBytesRemaining());
 		}
 
@@ -917,7 +917,7 @@ namespace Electra
 				return;
 			}
 			Parsed = MakePimpl<FParsed>();
-			FMP4AtomReaderHelper ar(BoxInfo.Data);
+			FMP4AtomReader ar(BoxInfo.Data);
 			Parsed->HEVCDecoderConfigurationRecord = MakeConstArrayView(ar.GetCurrentDataPointer(), ar.GetNumBytesRemaining());
 		}
 
@@ -940,7 +940,7 @@ namespace Electra
 				return;
 			}
 			Parsed = MakePimpl<FParsed>();
-			FMP4AtomReaderHelper ar(BoxInfo.Data);
+			FMP4AtomReader ar(BoxInfo.Data);
 			Parsed->DOVIDecoderConfigurationRecord = MakeConstArrayView(ar.GetCurrentDataPointer(), ar.GetNumBytesRemaining());
 		}
 
@@ -963,7 +963,7 @@ namespace Electra
 				return;
 			}
 			Parsed = MakePimpl<FParsed>();
-			FMP4AtomReaderHelper ar(BoxInfo.Data);
+			FMP4AtomReader ar(BoxInfo.Data);
 			Parsed->AC3SpecificBox = MakeConstArrayView(ar.GetCurrentDataPointer(), ar.GetNumBytesRemaining());
 		}
 
@@ -986,7 +986,7 @@ namespace Electra
 				return;
 			}
 			Parsed = MakePimpl<FParsed>();
-			FMP4AtomReaderHelper ar(BoxInfo.Data);
+			FMP4AtomReader ar(BoxInfo.Data);
 			Parsed->EC3SpecificBox = MakeConstArrayView(ar.GetCurrentDataPointer(), ar.GetNumBytesRemaining());
 		}
 
@@ -1009,7 +1009,7 @@ namespace Electra
 				return;
 			}
 			Parsed = MakePimpl<FParsed>();
-			FMP4AtomReaderHelper ar(BoxInfo.Data);
+			FMP4AtomReader ar(BoxInfo.Data);
 			Parsed->FLACSpecificBox = MakeConstArrayView(ar.GetCurrentDataPointer(), ar.GetNumBytesRemaining());
 		}
 
@@ -1032,7 +1032,7 @@ namespace Electra
 				return;
 			}
 			Parsed = MakePimpl<FParsed>();
-			FMP4AtomReaderHelper ar(BoxInfo.Data);
+			FMP4AtomReader ar(BoxInfo.Data);
 			Parsed->OpusSpecificBox = MakeConstArrayView(ar.GetCurrentDataPointer(), ar.GetNumBytesRemaining());
 		}
 
@@ -1055,7 +1055,7 @@ namespace Electra
 				return;
 			}
 			Parsed = MakePimpl<FParsed>();
-			FMP4AtomReaderHelper ar(BoxInfo.Data);
+			FMP4AtomReader ar(BoxInfo.Data);
 			Parsed->VPCodecConfigurationBox = MakeConstArrayView(ar.GetCurrentDataPointer(), ar.GetNumBytesRemaining());
 		}
 
@@ -1078,7 +1078,7 @@ namespace Electra
 				return;
 			}
 			Parsed = MakePimpl<FParsed>();
-			FMP4AtomReaderHelper ar(BoxInfo.Data);
+			FMP4AtomReader ar(BoxInfo.Data);
 			Parsed->InitialObjectDescriptor = MakeConstArrayView(ar.GetCurrentDataPointer(), ar.GetNumBytesRemaining());
 		}
 
@@ -1103,7 +1103,7 @@ namespace Electra
 				return;
 			}
 			Parsed = MakePimpl<FParsed>();
-			FMP4AtomReaderHelper ar(BoxInfo.Data);
+			FMP4AtomReader ar(BoxInfo.Data);
 			ar.ReadVersionAndFlags(Parsed->Version, Parsed->Flags);
 			Parsed->ESDescriptor = MakeConstArrayView(ar.GetCurrentDataPointer(), ar.GetNumBytesRemaining());
 		}
@@ -1128,7 +1128,7 @@ namespace Electra
 				return;
 			}
 			Parsed = MakePimpl<FParsed>();
-			FMP4AtomReaderHelper ar(BoxInfo.Data);
+			FMP4AtomReader ar(BoxInfo.Data);
 			uint32 colour_type;
 			ar.Read(colour_type);
 			if (colour_type == MakeBoxAtom('n','c','l','c') || colour_type == MakeBoxAtom('n','c','l','x'))
@@ -1176,7 +1176,7 @@ namespace Electra
 				return;
 			}
 			Parsed = MakePimpl<FParsed>();
-			FMP4AtomReaderHelper ar(BoxInfo.Data);
+			FMP4AtomReader ar(BoxInfo.Data);
 			ar.Read(Parsed->BufferSizeDB);
 			ar.Read(Parsed->MaxBitrate);
 			ar.Read(Parsed->AverageBitrate);
@@ -1215,7 +1215,7 @@ namespace Electra
 				return;
 			}
 			Parsed = MakePimpl<FParsed>();
-			FMP4AtomReaderHelper ar(BoxInfo.Data);
+			FMP4AtomReader ar(BoxInfo.Data);
 			ar.ReadVersionAndFlags(Parsed->Version, Parsed->Flags);
 			uint32 entry_count;
 			ar.Read(entry_count);
@@ -1261,7 +1261,7 @@ namespace Electra
 				return;
 			}
 			Parsed = MakePimpl<FParsed>();
-			FMP4AtomReaderHelper ar(BoxInfo.Data);
+			FMP4AtomReader ar(BoxInfo.Data);
 			ar.ReadVersionAndFlags(Parsed->Version, Parsed->Flags);
 			uint32 entry_count;
 			ar.Read(entry_count);
@@ -1330,7 +1330,7 @@ namespace Electra
 				return;
 			}
 			Parsed = MakePimpl<FParsed>();
-			FMP4AtomReaderHelper ar(BoxInfo.Data);
+			FMP4AtomReader ar(BoxInfo.Data);
 			ar.ReadVersionAndFlags(Parsed->Version, Parsed->Flags);
 			uint32 entry_count;
 			ar.Read(entry_count);
@@ -1369,7 +1369,7 @@ namespace Electra
 				return;
 			}
 			Parsed = MakePimpl<FParsed>();
-			FMP4AtomReaderHelper ar(BoxInfo.Data);
+			FMP4AtomReader ar(BoxInfo.Data);
 			ar.ReadVersionAndFlags(Parsed->Version, Parsed->Flags);
 			ar.Read(Parsed->ConstantSampleSize);
 			ar.Read(Parsed->NumberOfSamples);
@@ -1404,7 +1404,7 @@ namespace Electra
 				return;
 			}
 			Parsed = MakePimpl<FParsed>();
-			FMP4AtomReaderHelper ar(BoxInfo.Data);
+			FMP4AtomReader ar(BoxInfo.Data);
 			ar.ReadVersionAndFlags(Parsed->Version, Parsed->Flags);
 			ar.Read(Parsed->NumberOfEntries);
 			if (BoxInfo.Type == MakeBoxAtom('c','o','6','4'))
@@ -1447,7 +1447,7 @@ namespace Electra
 				return;
 			}
 			Parsed = MakePimpl<FParsed>();
-			FMP4AtomReaderHelper ar(BoxInfo.Data);
+			FMP4AtomReader ar(BoxInfo.Data);
 			ar.ReadVersionAndFlags(Parsed->Version, Parsed->Flags);
 			if ((Parsed->Flags & 1) != 0)
 			{
@@ -1486,7 +1486,7 @@ namespace Electra
 				return;
 			}
 			Parsed = MakePimpl<FParsed>();
-			FMP4AtomReaderHelper ar(BoxInfo.Data);
+			FMP4AtomReader ar(BoxInfo.Data);
 			ar.ReadVersionAndFlags(Parsed->Version, Parsed->Flags);
 			if ((Parsed->Flags & 1) != 0)
 			{
@@ -1532,7 +1532,7 @@ namespace Electra
 				return;
 			}
 			Parsed = MakePimpl<FParsed>();
-			FMP4AtomReaderHelper ar(BoxInfo.Data);
+			FMP4AtomReader ar(BoxInfo.Data);
 			ar.ReadVersionAndFlags(Parsed->Version, Parsed->Flags);
 			if (Parsed->Version == 0)
 			{
@@ -1605,7 +1605,7 @@ namespace Electra
 				return;
 			}
 			Parsed = MakePimpl<FParsed>();
-			FMP4AtomReaderHelper ar(BoxInfo.Data);
+			FMP4AtomReader ar(BoxInfo.Data);
 			ar.ReadVersionAndFlags(Parsed->Version, Parsed->Flags);
 			ar.Read(Parsed->GroupingType);
 		#if !UE_BUILD_SHIPPING
@@ -1667,7 +1667,7 @@ namespace Electra
 				return;
 			}
 			Parsed = MakePimpl<FParsed>();
-			FMP4AtomReaderHelper ar(BoxInfo.Data);
+			FMP4AtomReader ar(BoxInfo.Data);
 			ar.ReadVersionAndFlags(Parsed->Version, Parsed->Flags);
 			if (Parsed->Version == 1)
 			{
@@ -1706,7 +1706,7 @@ namespace Electra
 				return;
 			}
 			Parsed = MakePimpl<FParsed>();
-			FMP4AtomReaderHelper ar(BoxInfo.Data);
+			FMP4AtomReader ar(BoxInfo.Data);
 			ar.ReadVersionAndFlags(Parsed->Version, Parsed->Flags);
 			ar.Read(Parsed->TrackID);
 			ar.Read(Parsed->DefaultSampleDescriptionIndex);
@@ -1756,7 +1756,7 @@ namespace Electra
 				return;
 			}
 			Parsed = MakePimpl<FParsed>();
-			FMP4AtomReaderHelper ar(BoxInfo.Data);
+			FMP4AtomReader ar(BoxInfo.Data);
 			ar.ReadVersionAndFlags(Parsed->Version, Parsed->Flags);
 			ar.Read(Parsed->SequenceNumber);
 		}
@@ -1787,7 +1787,7 @@ namespace Electra
 				return;
 			}
 			Parsed = MakePimpl<FParsed>();
-			FMP4AtomReaderHelper ar(BoxInfo.Data);
+			FMP4AtomReader ar(BoxInfo.Data);
 			ar.ReadVersionAndFlags(Parsed->Version, Parsed->Flags);
 			ar.Read(Parsed->TrackID);
 			if ((Parsed->Flags & 1) != 0)
@@ -1900,7 +1900,7 @@ namespace Electra
 				return;
 			}
 			Parsed = MakePimpl<FParsed>();
-			FMP4AtomReaderHelper ar(BoxInfo.Data);
+			FMP4AtomReader ar(BoxInfo.Data);
 			ar.ReadVersionAndFlags(Parsed->Version, Parsed->Flags);
 			if (Parsed->Version == 1)
 			{
@@ -1941,7 +1941,7 @@ namespace Electra
 				return;
 			}
 			Parsed = MakePimpl<FParsed>();
-			FMP4AtomReaderHelper ar(BoxInfo.Data);
+			FMP4AtomReader ar(BoxInfo.Data);
 			ar.ReadVersionAndFlags(Parsed->Version, Parsed->Flags);
 			ar.Read(Parsed->SampleCount);
 			if ((Parsed->Flags & 1) != 0)
@@ -2096,7 +2096,7 @@ namespace Electra
 				return;
 			}
 			Parsed = MakePimpl<FParsed>();
-			FMP4AtomReaderHelper ar(BoxInfo.Data);
+			FMP4AtomReader ar(BoxInfo.Data);
 			ar.ReadVersionAndFlags(Parsed->Version, Parsed->Flags);
 			Parsed->SystemID = MakeConstArrayView(ar.GetCurrentDataPointer(), 16);
 			ar.SkipBytes(16);
@@ -2151,7 +2151,7 @@ namespace Electra
 				return;
 			}
 			Parsed = MakePimpl<FParsed>();
-			FMP4AtomReaderHelper ar(BoxInfo.Data);
+			FMP4AtomReader ar(BoxInfo.Data);
 			ar.ReadVersionAndFlags(Parsed->Version, Parsed->Flags);
 			if (Parsed->Version == 0)
 			{
@@ -2234,7 +2234,7 @@ namespace Electra
 				return;
 			}
 			Parsed = MakePimpl<FParsed>();
-			FMP4AtomReaderHelper ar(BoxInfo.Data);
+			FMP4AtomReader ar(BoxInfo.Data);
 			ar.ReadVersionAndFlags(Parsed->Version, Parsed->Flags);
 			if (InitRefs.IsValid())
 			{
@@ -2358,7 +2358,7 @@ namespace Electra
 				return;
 			}
 			Parsed = MakePimpl<FParsed>();
-			FMP4AtomReaderHelper ar(BoxInfo.Data);
+			FMP4AtomReader ar(BoxInfo.Data);
 			ar.ReadVersionAndFlags(Parsed->Version, Parsed->Flags);
 			ar.Read(Parsed->ReferenceID);
 			ar.Read(Parsed->Timescale);
@@ -2438,7 +2438,7 @@ namespace Electra
 				return;
 			}
 			Parsed = MakePimpl<FParsed>();
-			FMP4AtomReaderHelper ar(BoxInfo.Data);
+			FMP4AtomReader ar(BoxInfo.Data);
 			ar.ReadVersionAndFlags(Parsed->Version, Parsed->Flags);
 			if (Parsed->Version == 1)
 			{
diff --git a/Engine/Plugins/Media/ElectraUtil/Source/ElectraBase/Private/Utilities/UtilitiesMP4.cpp b/Engine/Plugins/Media/ElectraUtil/Source/ElectraBase/Private/Utilities/UtilitiesMP4.cpp
--- a/Engine/Plugins/Media/ElectraUtil/Source/ElectraBase/Private/Utilities/UtilitiesMP4.cpp
+++ b/Engine/Plugins/Media/ElectraUtil/Source/ElectraBase/Private/Utilities/UtilitiesMP4.cpp
@@ -11,31 +11,32 @@
 
 namespace Electra
 {
+
 namespace UtilitiesMP4
 {
 
 
-FMP4AtomReaderHelper::FMP4AtomReaderHelper(const TConstArrayView<uint8>& InData)
+FMP4AtomReader::FMP4AtomReader(const TConstArrayView<uint8>& InData)
 	: DataPtr(InData.GetData()), DataSize(InData.Num()), CurrentOffset(0)
 {
 }
 
-int64 FMP4AtomReaderHelper::GetCurrentOffset() const
+int64 FMP4AtomReader::GetCurrentOffset() const
 {
 	return CurrentOffset;
 }
 
-int64 FMP4AtomReaderHelper::GetNumBytesRemaining() const
+int64 FMP4AtomReader::GetNumBytesRemaining() const
 {
 	return DataSize - GetCurrentOffset();
 }
 
-const uint8* FMP4AtomReaderHelper::GetCurrentDataPointer() const
+const uint8* FMP4AtomReader::GetCurrentDataPointer() const
 {
 	return GetNumBytesRemaining() ? DataPtr + GetCurrentOffset() : nullptr;
 }
 
-void FMP4AtomReaderHelper::SetCurrentOffset(int64 InNewOffset)
+void FMP4AtomReader::SetCurrentOffset(int64 InNewOffset)
 {
 	check(InNewOffset >= 0 && InNewOffset <= DataSize);
 	if (InNewOffset >= 0 && InNewOffset <= DataSize)
@@ -44,7 +45,7 @@ void FMP4AtomReaderHelper::SetCurrentOffset(int64 InNewOffset)
 	}
 }
 
-bool FMP4AtomReaderHelper::ReadVersionAndFlags(uint8& OutVersion, uint32& OutFlags)
+bool FMP4AtomReader::ReadVersionAndFlags(uint8& OutVersion, uint32& OutFlags)
 {
 	uint32 VersionAndFlags = 0;
 	if (!Read(VersionAndFlags))
@@ -56,7 +57,7 @@ bool FMP4AtomReaderHelper::ReadVersionAndFlags(uint8& OutVersion, uint32& OutFla
 	return true;
 }
 
-bool FMP4AtomReaderHelper::ReadString(FString& OutString, uint16 InNumBytes)
+bool FMP4AtomReader::ReadString(FString& OutString, uint16 InNumBytes)
 {
 	OutString.Empty();
 	if (InNumBytes == 0)
@@ -80,7 +81,7 @@ bool FMP4AtomReaderHelper::ReadString(FString& OutString, uint16 InNumBytes)
 	return false;
 }
 
-bool FMP4AtomReaderHelper::ReadStringUTF8(FString& OutString, int32 InNumBytes)
+bool FMP4AtomReader::ReadStringUTF8(FString& OutString, int32 InNumBytes)
 {
 	OutString.Empty();
 	if (InNumBytes == 0)
@@ -107,7 +108,7 @@ bool FMP4AtomReaderHelper::ReadStringUTF8(FString& OutString, int32 InNumBytes)
 	return false;
 }
 
-bool FMP4AtomReaderHelper::ReadStringUTF16(FString& OutString, int32 InNumBytes)
+bool FMP4AtomReader::ReadStringUTF16(FString& OutString, int32 InNumBytes)
 {
 	OutString.Empty();
 	if (InNumBytes == 0)
@@ -137,7 +138,7 @@ bool FMP4AtomReaderHelper::ReadStringUTF16(FString& OutString, int32 InNumBytes)
 	return false;
 }
 
-bool FMP4AtomReaderHelper::ReadAsNumber(uint64& OutValue, int32 InNumBytes)
+bool FMP4AtomReader::ReadAsNumber(uint64& OutValue, int32 InNumBytes)
 {
 	OutValue = 0;
 	if (InNumBytes < 0 || InNumBytes > 8)
@@ -155,7 +156,7 @@ bool FMP4AtomReaderHelper::ReadAsNumber(uint64& OutValue, int32 InNumBytes)
 	}
 	return true;
 }
-bool FMP4AtomReaderHelper::ReadAsNumber(int64& OutValue, int32 InNumBytes)
+bool FMP4AtomReader::ReadAsNumber(int64& OutValue, int32 InNumBytes)
 {
 	OutValue = 0;
 	if (InNumBytes < 0 || InNumBytes > 8)
@@ -177,7 +178,7 @@ bool FMP4AtomReaderHelper::ReadAsNumber(int64& OutValue, int32 InNumBytes)
 	}
 	return true;
 }
-bool FMP4AtomReaderHelper::ReadAsNumber(float& OutValue)
+bool FMP4AtomReader::ReadAsNumber(float& OutValue)
 {
 	uint32 Flt;
 	if (Read(Flt))
@@ -187,7 +188,7 @@ bool FMP4AtomReaderHelper::ReadAsNumber(float& OutValue)
 	}
 	return false;
 }
-bool FMP4AtomReaderHelper::ReadAsNumber(double& OutValue)
+bool FMP4AtomReader::ReadAsNumber(double& OutValue)
 {
 	uint64 Dbl;
 	if (Read(Dbl))
@@ -198,12 +199,12 @@ bool FMP4AtomReaderHelper::ReadAsNumber(double& OutValue)
 	return false;
 }
 
-bool FMP4AtomReaderHelper::ReadBytes(void* Buffer, int32 InNumBytes)
+bool FMP4AtomReader::ReadBytes(void* Buffer, int32 InNumBytes)
 {
 	return ReadData(Buffer, InNumBytes) == InNumBytes;
 }
 
-int32 FMP4AtomReaderHelper::ReadData(void* IntoBuffer, int32 NumBytesToRead)
+int32 FMP4AtomReader::ReadData(void* IntoBuffer, int32 NumBytesToRead)
 {
 	if (NumBytesToRead <= 0)
 	{
@@ -222,7 +223,7 @@ int32 FMP4AtomReaderHelper::ReadData(void* IntoBuffer, int32 NumBytesToRead)
 	return -1;
 }
 
-bool FMP4AtomReaderHelper::ParseIntoBoxInfo(FMP4BoxInfo& OutBoxInfo, int64 InAtFileOffset)
+bool FMP4AtomReader::ParseIntoBoxInfo(FMP4BoxInfo& OutBoxInfo, int64 InAtFileOffset)
 {
 	// Clear output with default values before continuing.
 	OutBoxInfo = FMP4BoxInfo();
@@ -404,7 +405,7 @@ bool FMP4BoxTreeParser::ParseBoxTreeInternal(const TWeakPtr<FMP4BoxBase>& InPare
 		int64 NextBoxOffset = BoxTree->GetBoxFileOffset() + BoxTree->GetBoxDataOffset();
 		while(BoxBytesRemaining > 0)
 		{
-			FMP4AtomReaderHelper bh(MakeConstArrayView(BoxData, BoxBytesRemaining));
+			FMP4AtomReader bh(MakeConstArrayView(BoxData, BoxBytesRemaining));
 			FMP4BoxInfo bi;
 			if (!bh.ParseIntoBoxInfo(bi, NextBoxOffset))
 			{
@@ -554,6 +555,4 @@ FString FFileDataReader::GetLastError()
 	return LastError;
 }
 
-
-
 } // namespace Electra
diff --git a/Engine/Plugins/Media/ElectraUtil/Source/ElectraBase/Public/MediaVideoDecoderOutput.h b/Engine/Plugins/Media/ElectraUtil/Source/ElectraBase/Public/MediaVideoDecoderOutput.h
--- a/Engine/Plugins/Media/ElectraUtil/Source/ElectraBase/Public/MediaVideoDecoderOutput.h
+++ b/Engine/Plugins/Media/ElectraUtil/Source/ElectraBase/Public/MediaVideoDecoderOutput.h
@@ -1,48 +1,11 @@
 // Copyright Epic Games, Inc. All Rights Reserved.
 #pragma once
 
-#include "Misc/Timespan.h"
-#include "Math/IntPoint.h"
+#include "CoreTypes.h"
 #include "Templates/SharedPointer.h"
-#include "Templates/RefCounting.h"
-#include "PixelFormat.h"
-
-#include "MediaDecoderOutput.h"
-#include "ParameterDictionary.h"
 
 // -------------------------------------------------------------------------------------------------------------------------------
 
-struct FVideoDecoderCropInfo
-{
-	FVideoDecoderCropInfo() : CropLeft(0), CropTop(0), CropRight(0), CropBottom(0) {}
-	FVideoDecoderCropInfo(uint32 InCropLeft, uint32 InCropTop, uint32 InCropRight, uint32 InCropBottom) : CropLeft(InCropLeft), CropTop(InCropRight), CropRight(InCropRight), CropBottom(InCropBottom) {}
-
-	bool HasCropping() const { return CropLeft != 0 || CropRight != 0 || CropTop != 0 || CropBottom != 0;  }
-
-	uint32 CropLeft;
-	uint32 CropTop;
-	uint32 CropRight;
-	uint32 CropBottom;
-};
-
-enum class EVideoOrientation
-{
-	Original = 0,
-	CW90,
-	CW180,
-	CW270,
-};
-
-class IVideoDecoderTexture : public TSharedFromThis<IVideoDecoderTexture, ESPMode::ThreadSafe>
-{
-public:
-	virtual ~IVideoDecoderTexture() {}
-protected:
-	IVideoDecoderTexture() {}
-};
-
-
-
 class IVideoDecoderTimecode : public TSharedFromThis<IVideoDecoderTimecode, ESPMode::ThreadSafe>
 {
 public:
@@ -136,241 +99,3 @@ public:
 	// Get content light level info if available. Returns nullptr if information is not available.
 	virtual const FVideoDecoderHDRMetadata_content_light_level_info* GetContentLightLevelInfo() const = 0;
 };
-
-
-//! Interpretation of delivered pixel data
-enum class EVideoDecoderPixelEncoding
-{
-	Native = 0,		//!< Pixel formats native representation
-	RGB,			//!< Interpret as RGB
-	RGBA,			//!< Interpret as RGBA
-	YCbCr,			//!< Interpret as YCbCR
-	YCbCr_Alpha,	//!< Interpret as YCbCR with alpha
-	YCoCg,			//!< Interpret as scaled YCoCg
-	YCoCg_Alpha,	//!< Interpret as scaled YCoCg with trailing BC4 alpha data
-	CbY0CrY1,		//!< Interpret as CbY0CrY1
-	Y0CbY1Cr,		//!< Interpret as Y0CbY1Cr
-	ARGB_BigEndian,	//!< Interpret as ARGB, big endian
-};
-
-
-class FVideoDecoderOutput : public IDecoderOutput
-{
-public:
-	virtual ~FVideoDecoderOutput() = default;
-
-	virtual const Electra::FParamDict& GetDict() const
-	{
-		check(ParamDict.IsValid());
-		return *ParamDict;
-	}
-
-	virtual FDecoderTimeStamp GetTime() const
-	{
-		if (!ParamDict)
-		{
-			return FDecoderTimeStamp(FTimespan::Zero(), 0);
-		}
-		if (!(Cached.Flags & FCached::Valid_PTS))
-		{
-			Electra::FTimeValue pts(ParamDict->GetValue(IDecoderOutputOptionNames::PTS).GetTimeValue());
-			Cached.PTS = FDecoderTimeStamp(pts.GetAsTimespan(), pts.GetSequenceIndex());
-			FPlatformMisc::MemoryBarrier();
-			Cached.Flags |= FCached::Valid_PTS;
-		}
-		return Cached.PTS;
-	}
-
-	virtual void SetTime(const FDecoderTimeStamp& InTime)
-	{
-		Cached.PTS = FDecoderTimeStamp(InTime);
-		FPlatformMisc::MemoryBarrier();
-		Cached.Flags |= FCached::Valid_PTS;
-	}
-
-	virtual FTimespan GetDuration() const
-	{
-		if (!ParamDict)
-		{
-			return FTimespan(0);
-		}
-		if (!(Cached.Flags & FCached::Valid_Duration))
-		{
-			Cached.Duration = ParamDict->GetValue(IDecoderOutputOptionNames::Duration).GetTimeValue().GetAsTimespan();
-			FPlatformMisc::MemoryBarrier();
-			Cached.Flags |= FCached::Valid_Duration;
-		}
-		return Cached.Duration;
-	}
-
-	virtual FIntPoint GetOutputDim() const
-	{
-		if (!ParamDict)
-		{
-			return FIntPoint::ZeroValue;
-		}
-		if (!(Cached.Flags & FCached::Valid_OutputDim))
-		{
-			Cached.OutputDim.X = (int32)ParamDict->GetValue(IDecoderOutputOptionNames::Width).SafeGetInt64(0);
-			Cached.OutputDim.Y = (int32)ParamDict->GetValue(IDecoderOutputOptionNames::Height).SafeGetInt64(0);
-			FPlatformMisc::MemoryBarrier();
-			Cached.Flags |= FCached::Valid_OutputDim;
-		}
-		return Cached.OutputDim;
-	}
-
-	virtual double GetAspectRatio() const
-	{
-		if (!ParamDict)
-		{
-			return 1.0;
-		}
-		const FIntPoint Dim = GetOutputDim();
-		return ((double)Dim.X / (double)Dim.Y) * ParamDict->GetValue(IDecoderOutputOptionNames::AspectRatio).SafeGetDouble(1.0);
-	}
-
-	virtual FVideoDecoderCropInfo GetCropInfo() const
-	{
-		if (!ParamDict)
-		{
-			return FVideoDecoderCropInfo();
-		}
-		if (!(Cached.Flags & FCached::Valid_CropInfo))
-		{
-			Cached.CropInfo.CropLeft = (int32)ParamDict->GetValue(IDecoderOutputOptionNames::CropLeft).SafeGetInt64(0);
-			Cached.CropInfo.CropTop = (int32)ParamDict->GetValue(IDecoderOutputOptionNames::CropTop).SafeGetInt64(0);
-			Cached.CropInfo.CropRight = (int32)ParamDict->GetValue(IDecoderOutputOptionNames::CropRight).SafeGetInt64(0);
-			Cached.CropInfo.CropBottom = (int32)ParamDict->GetValue(IDecoderOutputOptionNames::CropBottom).SafeGetInt64(0);
-			FPlatformMisc::MemoryBarrier();
-			Cached.Flags |= FCached::Valid_CropInfo;
-		}
-		return Cached.CropInfo;
-	}
-
-	virtual FIntPoint GetDim() const
-	{
-		if (!ParamDict)
-		{
-			return FIntPoint::ZeroValue;
-		}
-		const FIntPoint OutputDim = GetOutputDim();
-		const FVideoDecoderCropInfo CropInfo = GetCropInfo();
-
-		FIntPoint Dim;
-		Dim.X = OutputDim.X + (CropInfo.CropLeft + CropInfo.CropRight);
-		Dim.Y = OutputDim.Y + (CropInfo.CropTop + CropInfo.CropBottom);
-		return Dim;
-	}
-
-	virtual EPixelFormat GetFormat() const
-	{
-		if (!ParamDict)
-		{
-			return EPixelFormat::PF_Unknown;
-		}
-		return (EPixelFormat)ParamDict->GetValue(IDecoderOutputOptionNames::PixelFormat).SafeGetInt64((int64)EPixelFormat::PF_Unknown);
-	}
-
-	virtual EVideoDecoderPixelEncoding GetFormatEncoding() const
-	{
-		if (!ParamDict)
-		{
-			return EVideoDecoderPixelEncoding::Native;
-		}
-		return (EVideoDecoderPixelEncoding)ParamDict->GetValue(IDecoderOutputOptionNames::PixelEncoding).SafeGetInt64((int64)EVideoDecoderPixelEncoding::Native);
-	}
-
-	virtual EVideoOrientation GetOrientation() const
-	{
-		if (!ParamDict)
-		{
-			return EVideoOrientation::Original;
-		}
-		if (!(Cached.Flags & FCached::Valid_Orientation))
-		{
-			Cached.Orientation = (EVideoOrientation)ParamDict->GetValue(IDecoderOutputOptionNames::Orientation).SafeGetInt64((int64)EVideoOrientation::Original);
-			FPlatformMisc::MemoryBarrier();
-			Cached.Flags |= FCached::Valid_Orientation;
-		}
-		return Cached.Orientation;
-	}
-
-	virtual int32 GetBitsPerComponent() const
-	{
-		const int32 NumDefaultBits = 8;
-		if (ParamDict)
-		{
-			return (int32)ParamDict->GetValue(IDecoderOutputOptionNames::BitsPerComponent).SafeGetInt64((int64)NumDefaultBits);
-		}
-		return NumDefaultBits;
-	}
-
-	virtual TSharedPtr<const IVideoDecoderHDRInformation, ESPMode::ThreadSafe> GetHDRInformation() const
-	{
-		if (ParamDict && ParamDict->HaveKey(IDecoderOutputOptionNames::HDRInfo))
-		{
-			return ParamDict->GetValue(IDecoderOutputOptionNames::HDRInfo).GetSharedPointer<const IVideoDecoderHDRInformation>();
-		}
-		return nullptr;
-	}
-
-	virtual TSharedPtr<const IVideoDecoderColorimetry, ESPMode::ThreadSafe> GetColorimetry() const
-	{
-		if (ParamDict && ParamDict->HaveKey(IDecoderOutputOptionNames::Colorimetry))
-		{
-			return ParamDict->GetValue(IDecoderOutputOptionNames::Colorimetry).GetSharedPointer<const IVideoDecoderColorimetry>();
-		}
-		return nullptr;
-	}
-
-	virtual TSharedPtr<const IVideoDecoderTimecode, ESPMode::ThreadSafe> GetTimecode() const
-	{
-		if (ParamDict && ParamDict->HaveKey(IDecoderOutputOptionNames::Timecode))
-		{
-			return ParamDict->GetValue(IDecoderOutputOptionNames::Timecode).GetSharedPointer<const IVideoDecoderTimecode>();
-		}
-		return nullptr;
-	}
-
-protected:
-	FVideoDecoderOutput() = default;
-
-	void Initialize(TSharedPtr<Electra::FParamDict, ESPMode::ThreadSafe> InParamDict)
-	{
-		Cached.Flags = 0;
-		ParamDict = MoveTemp(InParamDict);
-#if !UE_BUILD_SHIPPING
-		// For debugging purposes, get the PTS and duration so they're easily accessible.
-		(void)GetTime();
-		(void)GetDuration();
-#endif
-	}
-
-private:
-	struct FCached
-	{
-		FCached()
-		{
-			OutputDim = FIntPoint::ZeroValue;
-		}
-		enum
-		{
-			Valid_CropInfo = 1 << 0,
-			Valid_OutputDim = 1 << 1,
-			Valid_Orientation = 1 << 2,
-			Valid_PTS = 1 << 3,
-			Valid_Duration = 1 << 4
-		};
-
-		FVideoDecoderCropInfo CropInfo;
-		FDecoderTimeStamp PTS;
-		FTimespan Duration;
-		FIntPoint OutputDim;
-		EVideoOrientation Orientation = EVideoOrientation::Original;;
-		uint32 Flags = 0;
-	};
-	mutable FCached Cached;
-	TSharedPtr<Electra::FParamDict, ESPMode::ThreadSafe> ParamDict;
-};
-
-using FVideoDecoderOutputPtr = TSharedPtr<FVideoDecoderOutput, ESPMode::ThreadSafe>;
diff --git a/Engine/Plugins/Media/ElectraUtil/Source/ElectraBase/Public/Utilities/MP4Boxes/MP4BoxBase.h b/Engine/Plugins/Media/ElectraUtil/Source/ElectraBase/Public/Utilities/MP4Boxes/MP4BoxBase.h
--- a/Engine/Plugins/Media/ElectraUtil/Source/ElectraBase/Public/Utilities/MP4Boxes/MP4BoxBase.h
+++ b/Engine/Plugins/Media/ElectraUtil/Source/ElectraBase/Public/Utilities/MP4Boxes/MP4BoxBase.h
@@ -112,7 +112,7 @@ namespace Electra
 				return false;
 			}
 
-			virtual void ProcessBoxChildrenRecursively(FMP4AtomReaderHelper& InOutReader, const FMP4BoxInfo& InCurrentBoxInfo);
+			virtual void ProcessBoxChildrenRecursively(FMP4AtomReader& InOutReader, const FMP4BoxInfo& InCurrentBoxInfo);
 
 			void SetRootBoxData(const TSharedPtr<FMP4BoxInfo, ESPMode::ThreadSafe>& InRootBoxData)
 			{ RootBoxData = InRootBoxData; }
diff --git a/Engine/Plugins/Media/ElectraUtil/Source/ElectraBase/Public/Utilities/UtilitiesMP4.h b/Engine/Plugins/Media/ElectraUtil/Source/ElectraBase/Public/Utilities/UtilitiesMP4.h
--- a/Engine/Plugins/Media/ElectraUtil/Source/ElectraBase/Public/Utilities/UtilitiesMP4.h
+++ b/Engine/Plugins/Media/ElectraUtil/Source/ElectraBase/Public/Utilities/UtilitiesMP4.h
@@ -49,11 +49,6 @@ namespace Electra
 			return (static_cast<uint32>(A) << 24) | (static_cast<uint32>(B) << 16) | (static_cast<uint32>(C) << 8) | static_cast<uint32>(D);
 		}
 
-		static constexpr uint32 MakeBoxAtom(const char A, const char B, const char C, const char D)
-		{
-			return (static_cast<uint32>(A) << 24) | (static_cast<uint32>(B) << 16) | (static_cast<uint32>(C) << 8) | static_cast<uint32>(D);
-		}
-
 		static FString GetPrintableBoxAtom(uint32 InAtom)
 		{
 			TCHAR tc[4];
@@ -106,10 +101,10 @@ namespace Electra
 		};
 
 
-		class FMP4AtomReaderHelper
+		class FMP4AtomReader
 		{
 		public:
-			UE_API FMP4AtomReaderHelper(const TConstArrayView<uint8>& InData);
+			UE_API FMP4AtomReader(const TConstArrayView<uint8>& InData);
 
 			UE_API bool ParseIntoBoxInfo(FMP4BoxInfo& OutBoxInfo, int64 InAtFileOffset);
 
@@ -185,6 +180,7 @@ namespace Electra
 	} // namespace UtilitiesMP4
 
 
+
 	class IFileDataReader : public IBaseDataReader
 	{
 	public:
diff --git a/Engine/Plugins/Media/ElectraUtil/Source/ElectraSamples/Private/ElectraAudioSample.cpp b/Engine/Plugins/Media/ElectraUtil/Source/ElectraSamples/Private/ElectraAudioSample.cpp
--- a/Engine/Plugins/Media/ElectraUtil/Source/ElectraSamples/Private/ElectraAudioSample.cpp
+++ b/Engine/Plugins/Media/ElectraUtil/Source/ElectraSamples/Private/ElectraAudioSample.cpp
@@ -7,7 +7,7 @@ FElectraAudioSample::~FElectraAudioSample()
 	FMemory::Free(Buffer);
 }
 
-bool FElectraAudioSample::Initialize(const void* InData, EMediaAudioSampleFormat InFormat, uint32 InNumChannels, uint32 InNumFrames, uint32 InSampleRate, const FMediaTimeStamp& InTime, const FTimespan& InDuration)
+bool FElectraAudioSample::AllocateFor(EMediaAudioSampleFormat InFormat, uint32 InNumChannels, uint32 InNumFrames)
 {
 	check(InFormat == EMediaAudioSampleFormat::Float);
 
@@ -25,19 +25,20 @@ bool FElectraAudioSample::Initialize(const void* InData, EMediaAudioSampleFormat
 	MediaAudioSampleFormat = InFormat;
 	NumChannels = InNumChannels;
 	NumFrames = InNumFrames;
+	return true;
+}
+
+void FElectraAudioSample::SetParameters(uint32 InSampleRate, const FMediaTimeStamp& InTime, const FTimespan& InDuration)
+{
 	SampleRate = InSampleRate;
 	MediaTimeStamp = InTime;
 	Duration = InDuration;
-	if (Buffer && InData)
-	{
-		FMemory::Memcpy(Buffer, InData, NumBytesNeeded);
-	}
-	return true;
 }
 
 #if !UE_SERVER
 void FElectraAudioSample::ShutdownPoolable()
 {
-	// Nothing to do here at the moment.
+	ReleaseDelegate.ExecuteIfBound(this);
+	ReleaseDelegate.Unbind();
 }
 #endif
diff --git a/Engine/Plugins/Media/ElectraUtil/Source/ElectraSamples/Private/ElectraSamplesModule.cpp b/Engine/Plugins/Media/ElectraUtil/Source/ElectraSamples/Private/ElectraSamplesModule.cpp
--- a/Engine/Plugins/Media/ElectraUtil/Source/ElectraSamples/Private/ElectraSamplesModule.cpp
+++ b/Engine/Plugins/Media/ElectraUtil/Source/ElectraSamples/Private/ElectraSamplesModule.cpp
@@ -2,13 +2,10 @@
 
 #include "ElectraSamplesModule.h"
 
-#include "CoreMinimal.h"
+#include "CoreTypes.h"
 #include "Modules/ModuleManager.h"
 #include "IElectraSamplesModule.h"
 
-#include "IMediaTextureSample.h"
-#include "MediaVideoDecoderOutput.h"
-
 #define LOCTEXT_NAMESPACE "ElectraBaseModule"
 
 DEFINE_LOG_CATEGORY(LogElectraSamples);
@@ -17,25 +14,17 @@ DEFINE_LOG_CATEGORY(LogElectraSamples);
 
 class FElectraSamplesModule: public IElectraSamplesModule
 {
-public:
 public:
 	// IModuleInterface interface
-
 	void StartupModule() override
 	{
-		static_assert((int32)EMediaOrientation::Original == (int32)EVideoOrientation::Original, "check alignment of both enums");
-		static_assert((int32)EMediaOrientation::CW90 == (int32)EVideoOrientation::CW90, "check alignment of both enums");
 	}
 
 	void ShutdownModule() override
 	{
 	}
-
-private:
 };
 
 IMPLEMENT_MODULE(FElectraSamplesModule, ElectraSamples);
 
 #undef LOCTEXT_NAMESPACE
-
-
diff --git a/Engine/Plugins/Media/ElectraUtil/Source/ElectraSamples/Private/ElectraTextureSampleBase.cpp b/Engine/Plugins/Media/ElectraUtil/Source/ElectraSamples/Private/ElectraTextureSampleBase.cpp
--- a/Engine/Plugins/Media/ElectraUtil/Source/ElectraSamples/Private/ElectraTextureSampleBase.cpp
+++ b/Engine/Plugins/Media/ElectraUtil/Source/ElectraSamples/Private/ElectraTextureSampleBase.cpp
@@ -10,6 +10,9 @@
 #include "ElectraTextureSampleUtils.h"
 #include "ElectraSamplesModule.h"
 
+#include "MediaDecoderOutput.h"
+#include "MediaVideoDecoderOutput.h"
+
 // -------------------------------------------------------------------------------------------------------------------------------------------------------
 
 namespace
@@ -46,59 +49,64 @@ static TAutoConsoleVariable<int32> CVarElectraHdrToneMapMethod(
 
 // -------------------------------------------------------------------------------------------------------------------------------------------------------
 
-void IElectraTextureSampleBase::Initialize(FVideoDecoderOutput* InVideoDecoderOutput)
+void IElectraTextureSampleBase::InitializeCommon(const Electra::FParamDict& InParams)
 {
-	VideoDecoderOutput = StaticCastSharedPtr<FVideoDecoderOutput, IDecoderOutputPoolable, ESPMode::ThreadSafe>(InVideoDecoderOutput->AsShared());
-	Colorimetry = VideoDecoderOutput->GetColorimetry();
-	HDRInfo = VideoDecoderOutput->GetHDRInformation();
-	TSharedPtr<const IVideoDecoderTimecode, ESPMode::ThreadSafe> TimecodePtr = VideoDecoderOutput->GetTimecode();
-	if (TimecodePtr.IsValid())
+	// HDR info
+	if (InParams.HaveKey(IDecoderOutputOptionNames::HDRInfo))
+	{
+		HDRInfo = InParams.GetValue(IDecoderOutputOptionNames::HDRInfo).GetSharedPointer<const IVideoDecoderHDRInformation>();
+	}
+	// Colorimetry info
+	if (InParams.HaveKey(IDecoderOutputOptionNames::Colorimetry))
 	{
-		// Store this in case this is needed again.
-		DecoderTimecode = TimecodePtr;
-		Timecode = CreateTimecodeFromMPEGDefinition(Framerate, TimecodePtr->GetMPEGDefinition());
+		Colorimetry = InParams.GetValue(IDecoderOutputOptionNames::Colorimetry).GetSharedPointer<const IVideoDecoderColorimetry>();
 	}
-	else if (VideoDecoderOutput->GetDict().HaveKey(IDecoderOutputOptionNames::TMCDTimecode))
+	// Timecode
+	if (InParams.HaveKey(IDecoderOutputOptionNames::Timecode))
 	{
-		Timecode = VideoDecoderOutput->GetDict().GetValue(IDecoderOutputOptionNames::TMCDTimecode).GetTimecode();
-		Framerate = VideoDecoderOutput->GetDict().GetValue(IDecoderOutputOptionNames::TMCDFramerate).SafeGetFramerate();
+		Timecode = CreateTimecodeFromMPEGDefinition(Framerate, InParams.GetValue(IDecoderOutputOptionNames::Timecode).GetSharedPointer<const IVideoDecoderTimecode>()->GetMPEGDefinition());
+	}
+	else if (InParams.HaveKey(IDecoderOutputOptionNames::TMCDTimecode))
+	{
+		Timecode = InParams.GetValue(IDecoderOutputOptionNames::TMCDTimecode).GetTimecode();
+		Framerate = InParams.GetValue(IDecoderOutputOptionNames::TMCDFramerate).SafeGetFramerate();
 	}
 
-	// Get various basic MP4-style colorimetry values (we default to video range Rec709 SDR)
-	bool bFullRange = false;
+	// Get various basic MPEG-style colorimetry values (we default to video range Rec709 SDR)
 	uint8 ColorPrimaries = ElectraColorimetryUtils::DefaultMPEGColorPrimaries;
 	uint8 TransferCharacteristics = ElectraColorimetryUtils::DefaultMPEGMatrixCoefficients;
 	uint8 MatrixCoefficients = ElectraColorimetryUtils::DefaultMPEGTransferCharacteristics;
-	if (auto PinnedColorimetry = Colorimetry.Pin())
+	if (Colorimetry.IsValid())
 	{
-		bFullRange = (PinnedColorimetry->GetMPEGDefinition()->VideoFullRangeFlag != 0);
-		ColorPrimaries = PinnedColorimetry->GetMPEGDefinition()->ColourPrimaries;
-		TransferCharacteristics = PinnedColorimetry->GetMPEGDefinition()->TransferCharacteristics;
-		MatrixCoefficients = PinnedColorimetry->GetMPEGDefinition()->MatrixCoefficients;
+		bIsFullRange = Colorimetry->GetMPEGDefinition()->VideoFullRangeFlag != 0;
+		ColorPrimaries = Colorimetry->GetMPEGDefinition()->ColourPrimaries;
+		TransferCharacteristics = Colorimetry->GetMPEGDefinition()->TransferCharacteristics;
+		MatrixCoefficients = Colorimetry->GetMPEGDefinition()->MatrixCoefficients;
 	}
 
-	// Compute the bits per component in the data we get passed in
-	EPixelFormat PixFmt = VideoDecoderOutput->GetFormat();
+	PixelFormat = static_cast<EPixelFormat>(InParams.GetValue(IDecoderOutputOptionNames::PixelFormat).SafeGetInt64((int64)EPixelFormat::PF_Unknown));
 	uint8 NumBits = 8;
-	if (!IsDXTCBlockCompressedTextureFormat(PixFmt))
+	if (!IsDXTCBlockCompressedTextureFormat(PixelFormat))
 	{
-		if (PixFmt == PF_NV12)
+		if (PixelFormat == PF_NV12)
 		{
 			NumBits = 8;
 		}
-		else if (PixFmt == PF_A2B10G10R10)
+		else if (PixelFormat == PF_A2B10G10R10)
 		{
 			NumBits = 10;
 		}
-		else if (PixFmt == PF_P010)
+		else if (PixelFormat == PF_P010)
 		{
 			NumBits = 16;
 		}
 		else
 		{
-			NumBits = (8 * GPixelFormats[PixFmt].BlockBytes) / GPixelFormats[PixFmt].NumComponents;
+			NumBits = (8 * GPixelFormats[PixelFormat].BlockBytes) / GPixelFormats[PixelFormat].NumComponents;
 		}
 	}
+	PixelFormatEncoding = static_cast<EElectraTextureSamplePixelEncoding>(InParams.GetValue(IDecoderOutputOptionNames::PixelEncoding).SafeGetInt64((int64)EElectraTextureSamplePixelEncoding::Native));
+	PixelDataScale = static_cast<float>(InParams.GetValue(IDecoderOutputOptionNames::PixelDataScale).SafeGetDouble(1.0));
 
 	FVector Off = FVector::Zero();
 	const FMatrix* Mtx = nullptr;
@@ -111,14 +119,10 @@ void IElectraTextureSampleBase::Initialize(FVideoDecoderOutput* InVideoDecoderOu
 	MaxFALL = 0;
 
 	// Do we have specific HDR information, so we can assume a standard?
-	if (auto PinnedHDRInfo = HDRInfo.Pin())
+	if (HDRInfo.IsValid())
 	{
-		//
-		// HDR information present
-		//
-
 		// Mastering display info...
-		if (auto ColorVolume = PinnedHDRInfo->GetMasteringDisplayColourVolume())
+		if (auto ColorVolume = HDRInfo->GetMasteringDisplayColourVolume())
 		{
 			// A few sanity checks on the primaries coordinates (by no means exhaustive, but it should catch a fair share of oddities)
 			if (ColorVolume->display_primaries_x[0] > FMath::Max(ColorVolume->display_primaries_x[1], ColorVolume->display_primaries_x[2]) &&	// Red has largest X
@@ -139,7 +143,7 @@ void IElectraTextureSampleBase::Initialize(FVideoDecoderOutput* InVideoDecoderOu
 		}
 
 		// Content light level info...
-		if (auto ContentLightLevelInfo = PinnedHDRInfo->GetContentLightLevelInfo())
+		if (auto ContentLightLevelInfo = HDRInfo->GetContentLightLevelInfo())
 		{
 			MaxCLL = ContentLightLevelInfo->max_content_light_level;
 			MaxFALL =  ContentLightLevelInfo->max_pic_average_light_level;
@@ -150,20 +154,29 @@ void IElectraTextureSampleBase::Initialize(FVideoDecoderOutput* InVideoDecoderOu
 	SourceColorSpace = UE::Color::FColorSpace(ElectraColorimetryUtils::TranslateMPEGColorPrimaries(ColorPrimaries));
 
 	// Select the YUV-RGB conversion matrix to use
-	switch (ElectraColorimetryUtils::TranslateMPEGMatrixCoefficients(MatrixCoefficients))
+	switch(ElectraColorimetryUtils::TranslateMPEGMatrixCoefficients(MatrixCoefficients))
 	{
 		case UE::Color::EColorSpace::None:	// ID (RGB)
+		{
 			// no conversion, data is RGB
 			break;
+		}
 		case UE::Color::EColorSpace::sRGB:
-			Mtx = bFullRange ? &MediaShaders::YuvToRgbRec709Unscaled : &MediaShaders::YuvToRgbRec709Scaled;
+		{
+			Mtx = bIsFullRange ? &MediaShaders::YuvToRgbRec709Unscaled : &MediaShaders::YuvToRgbRec709Scaled;
 			break;
+		}
 		case UE::Color::EColorSpace::Rec2020:
-			Mtx = bFullRange ? &MediaShaders::YuvToRgbRec2020Unscaled : &MediaShaders::YuvToRgbRec2020Scaled;
+		{
+			Mtx = bIsFullRange ? &MediaShaders::YuvToRgbRec2020Unscaled : &MediaShaders::YuvToRgbRec2020Scaled;
 			break;
+		}
 		default:
+		{
 			check(!"*** Unexpected matrix coefficients!");
-			Mtx = bFullRange ? &MediaShaders::YuvToRgbRec709Unscaled : &MediaShaders::YuvToRgbRec709Scaled;
+			Mtx = bIsFullRange ? &MediaShaders::YuvToRgbRec709Unscaled : &MediaShaders::YuvToRgbRec709Scaled;
+			break;
+		}
 	}
 
 	// Get color encoding (sRGB, linear, PQ, HLG...)
@@ -172,13 +185,33 @@ void IElectraTextureSampleBase::Initialize(FVideoDecoderOutput* InVideoDecoderOu
 	if (Mtx)
 	{
 		// Select the offsets prior to YUV conversion needed per the incoming data
-		switch (NumBits)
+		switch(NumBits)
 		{
-		case 8:		Off = bFullRange ? MediaShaders::YUVOffsetNoScale8bits : MediaShaders::YUVOffset8bits; break;
-		case 10:	Off = bFullRange ? MediaShaders::YUVOffsetNoScale10bits : MediaShaders::YUVOffset10bits; break;
-		case 16:	Off = bFullRange ? MediaShaders::YUVOffsetNoScale16bits : MediaShaders::YUVOffset16bits; break;
-		case 32:	Off = bFullRange ? MediaShaders::YUVOffsetNoScaleFloat : MediaShaders::YUVOffsetFloat; break;
-		default:	check(!"Unexpected number of bits per channel!");
+			case 8:
+			{
+				Off = bIsFullRange ? MediaShaders::YUVOffsetNoScale8bits : MediaShaders::YUVOffset8bits;
+				break;
+			}
+			case 10:
+			{
+				Off = bIsFullRange ? MediaShaders::YUVOffsetNoScale10bits : MediaShaders::YUVOffset10bits;
+				break;
+			}
+			case 16:
+			{
+				Off = bIsFullRange ? MediaShaders::YUVOffsetNoScale16bits : MediaShaders::YUVOffset16bits;
+				break;
+			}
+			case 32:
+			{
+				Off = bIsFullRange ? MediaShaders::YUVOffsetNoScaleFloat : MediaShaders::YUVOffsetFloat;
+				break;
+			}
+			default:
+			{
+				check(!"Unexpected number of bits per channel!");
+				break;
+			}
 		}
 	}
 
@@ -187,7 +220,7 @@ void IElectraTextureSampleBase::Initialize(FVideoDecoderOutput* InVideoDecoderOu
 	float DataScale = GetSampleDataScale(NumBits == 10);
 
 	// Compute scale to make correct towards the max value (P010 will max out at 0xffc0 not 0xffff - so if it is present we need to adjust the scale a bit)
-	float NormScale = (VideoDecoderOutput->GetFormat() == PF_P010) ? (65535.0f / 65472.0f) : 1.0f;
+	float NormScale = PixelFormat == PF_P010 ? 65535.0f / 65472.0f : 1.0f;
 
 	// Matrix to transform sample data to standard YUV values
 	FMatrix PreMtx = FMatrix::Identity;
@@ -203,6 +236,43 @@ void IElectraTextureSampleBase::Initialize(FVideoDecoderOutput* InVideoDecoderOu
 
 	// Also store the plain YUV->RGB matrix (pointer) for later reference
 	YuvToRgbMtx = Mtx;
+
+
+	// Get the common values.
+	Electra::FTimeValue PTS(InParams.GetValue(IDecoderOutputOptionNames::PTS).GetTimeValue());
+	PresentationTimeStamp = FMediaTimeStamp(PTS.GetAsTimespan(), PTS.GetSequenceIndex());
+
+	Duration = InParams.GetValue(IDecoderOutputOptionNames::Duration).GetTimeValue().GetAsTimespan();
+
+	ImageOutputDim.X = (int32)InParams.GetValue(IDecoderOutputOptionNames::Width).SafeGetInt64(0);
+	ImageOutputDim.Y = (int32)InParams.GetValue(IDecoderOutputOptionNames::Height).SafeGetInt64(0);
+	AspectRatio = ((double)ImageOutputDim.X / (double)ImageOutputDim.Y) * InParams.GetValue(IDecoderOutputOptionNames::AspectRatio).SafeGetDouble(1.0);
+	Orientation = (EMediaOrientation)InParams.GetValue(IDecoderOutputOptionNames::Orientation).SafeGetInt64((int64)EMediaOrientation::Original);
+
+	CropLeft = (int32)InParams.GetValue(IDecoderOutputOptionNames::CropLeft).SafeGetInt64(0);
+	CropTop = (int32)InParams.GetValue(IDecoderOutputOptionNames::CropTop).SafeGetInt64(0);
+	CropRight = (int32)InParams.GetValue(IDecoderOutputOptionNames::CropRight).SafeGetInt64(0);
+	CropBottom = (int32)InParams.GetValue(IDecoderOutputOptionNames::CropBottom).SafeGetInt64(0);
+	ImageDim.X = ImageOutputDim.X + (CropLeft + CropRight);
+	ImageDim.Y = ImageOutputDim.Y + (CropTop + CropBottom);
+}
+
+bool IElectraTextureSampleBase::FinishInitialization()
+{
+	return true;
+}
+
+IElectraTextureSampleBase::FReleaseDelegate& IElectraTextureSampleBase::GetReleaseDelegate()
+{
+	return ReleaseDelegate;
+}
+
+
+// -------------------------------------------------------------------------------------------------------------------------------------------------------
+
+bool IElectraTextureSampleBase::IsCacheable() const
+{
+	return true;
 }
 
 void IElectraTextureSampleBase::InitializePoolable()
@@ -211,55 +281,84 @@ void IElectraTextureSampleBase::InitializePoolable()
 
 void IElectraTextureSampleBase::ShutdownPoolable()
 {
-	VideoDecoderOutput.Reset();
+	ReleaseDelegate.ExecuteIfBound(this);
+	// Clean out what is needed for this element to be reused.
+	ReleaseDelegate.Unbind();
+}
+
+bool IElectraTextureSampleBase::IsReadyForReuse()
+{
+	return true;
 }
 
 FIntPoint IElectraTextureSampleBase::GetDim() const
 {
-	if (VideoDecoderOutput)
-	{
-		return VideoDecoderOutput->GetDim();
-	}
-	return FIntPoint::ZeroValue;
+	return ImageDim;
 }
 
+void IElectraTextureSampleBase::SetDim(const FIntPoint& InDim)
+{
+	ImageDim = InDim;
+}
 
 FIntPoint IElectraTextureSampleBase::GetOutputDim() const
 {
-	if (VideoDecoderOutput)
-	{
-		return VideoDecoderOutput->GetOutputDim();
-	}
-	return FIntPoint::ZeroValue;
+	return ImageOutputDim;
 }
 
-
 FMediaTimeStamp IElectraTextureSampleBase::GetTime() const
 {
-	if (VideoDecoderOutput)
-	{
-		const FDecoderTimeStamp TimeStamp = VideoDecoderOutput->GetTime();
-		return FMediaTimeStamp(TimeStamp.Time, TimeStamp.SequenceIndex);
-	}
-	return FMediaTimeStamp();
+	return PresentationTimeStamp;
 }
 
 void IElectraTextureSampleBase::SetTime(const FMediaTimeStamp& InTime)
 {
-	if (VideoDecoderOutput)
-	{
-		VideoDecoderOutput->SetTime(FDecoderTimeStamp(InTime.Time, InTime.GetIndexValue()));
-	}
+	PresentationTimeStamp = InTime;
+}
+
+void IElectraTextureSampleBase::SetDuration(const FTimespan& InDuration)
+{
+	Duration = InDuration;
 }
 
+TSharedPtr<const IVideoDecoderColorimetry, ESPMode::ThreadSafe> IElectraTextureSampleBase::GetColorimetry() const
+{
+	return Colorimetry;
+}
+
+EPixelFormat IElectraTextureSampleBase::GetPixelFormat() const
+{
+	return PixelFormat;
+}
+
+EElectraTextureSamplePixelEncoding IElectraTextureSampleBase::GetPixelFormatEncoding() const
+{
+	return PixelFormatEncoding;
+}
 
 FTimespan IElectraTextureSampleBase::GetDuration() const
 {
-	if (VideoDecoderOutput)
-	{
-		return VideoDecoderOutput->GetDuration();
-	}
-	return FTimespan::Zero();
+	return Duration;
+}
+
+TOptional<FTimecode> IElectraTextureSampleBase::GetTimecode() const
+{
+	return Timecode;
+}
+
+TOptional<FFrameRate> IElectraTextureSampleBase::GetFramerate() const
+{
+	return Framerate;
+}
+
+double IElectraTextureSampleBase::GetAspectRatio() const
+{
+	return AspectRatio;
+}
+
+EMediaOrientation IElectraTextureSampleBase::GetOrientation() const
+{
+	return Orientation;
 }
 
 bool IElectraTextureSampleBase::IsOutputSrgb() const
@@ -274,11 +373,7 @@ const FMatrix& IElectraTextureSampleBase::GetYUVToRGBMatrix() const
 
 bool IElectraTextureSampleBase::GetFullRange() const
 {
-	if (auto PinnedColorimetry = Colorimetry.Pin())
-	{
-		return (PinnedColorimetry->GetMPEGDefinition()->VideoFullRangeFlag != 0);
-	}
-	return false;
+	return bIsFullRange;
 }
 
 FMatrix44f IElectraTextureSampleBase::GetSampleToRGBMatrix() const
@@ -314,7 +409,6 @@ bool IElectraTextureSampleBase::GetDisplayMasteringLuminance(float& OutMin, floa
 	{
 		return false;
 	}
-
 	OutMin = DisplayMasteringLuminanceMin;
 	OutMax = DisplayMasteringLuminanceMax;
 	return true;
@@ -331,7 +425,6 @@ bool IElectraTextureSampleBase::GetMaxLuminanceLevels(uint16& OutCLL, uint16& Ou
 	{
 		return false;
 	}
-
 	OutCLL = MaxCLL;
 	OutFALL = MaxFALL;
 	return true;
@@ -351,4 +444,9 @@ MediaShaders::EToneMapMethod IElectraTextureSampleBase::GetToneMapMethod() const
 	}
 }
 
+float IElectraTextureSampleBase::GetSampleDataScale(bool b10Bit) const
+{
+	return PixelDataScale;
+}
+
 #endif
diff --git a/Engine/Plugins/Media/ElectraUtil/Source/ElectraSamples/Public/Android/ElectraTextureSample.h b/Engine/Plugins/Media/ElectraUtil/Source/ElectraSamples/Public/Android/ElectraTextureSample.h
--- a/Engine/Plugins/Media/ElectraUtil/Source/ElectraSamples/Public/Android/ElectraTextureSample.h
+++ b/Engine/Plugins/Media/ElectraUtil/Source/ElectraSamples/Public/Android/ElectraTextureSample.h
@@ -13,7 +13,6 @@
 #include "Misc/Timespan.h"
 #include "MediaObjectPool.h"
 #include "MediaSamples.h"
-#include "MediaVideoDecoderOutputAndroid.h"
 
 #define UE_API ELECTRASAMPLES_API
 
@@ -24,97 +23,67 @@ struct AHardwareBuffer;
 /**
  * Texture sample generated by AndroidMedia player.
  */
-class FElectraTextureSample final
+class UE_API FElectraTextureSample final
 	: public IElectraTextureSampleBase
 	, public IMediaTextureSampleConverter
 	, public TSharedFromThis<FElectraTextureSample, ESPMode::ThreadSafe>
 {
 public:
+	FElectraTextureSample(TSharedPtr<FElectraTextureSampleSupport, ESPMode::ThreadSafe> InSupport);
+	virtual ~FElectraTextureSample();
 
-	/** Default constructor. */
-	FElectraTextureSample(TSharedPtr<FElectraTextureSampleSupport, ESPMode::ThreadSafe> InSupport)
-	: Support(InSupport)
-	, ImageResources(nullptr)
-	, Texture(nullptr)
-	, Buffer(nullptr)
-	, BufferSize(0)
-	{
-	}
-
-	/** Virtual destructor. */
-	UE_API virtual ~FElectraTextureSample();
+	bool Initialize();
 
-#if !UE_SERVER
-	UE_API virtual bool IsReadyForReuse() override;
-	UE_API virtual void ShutdownPoolable() override;
-#endif
-
-	UE_API void Initialize(FVideoDecoderOutput* InVideoDecoderOutput);
+	bool FinishInitialization() override;
 
 	//~ IMediaTextureSample interface
+	const void* GetBuffer() override;
+	uint32 GetStride() const override;
+	EMediaTextureSampleFormat GetFormat() const override;
+	FRHITexture* GetTexture() const override;
+	IMediaTextureSampleConverter* GetMediaTextureSampleConverter() override;
+	FIntPoint GetDim() const override;
 
-	virtual const void* GetBuffer() override
-	{
-		return Buffer;
-	}
-
-	UE_API virtual EMediaTextureSampleFormat GetFormat() const override;
-
-	UE_API virtual uint32 GetStride() const override;
-
-	virtual FRHITexture* GetTexture() const override
-	{
-		return Texture.GetReference();
-	}
-
-	virtual IMediaTextureSampleConverter* GetMediaTextureSampleConverter() override
-	{
-		return this;
-	}
+	bool IsReadyForReuse() override;
+	void ShutdownPoolable() override;
 
 private:
+	uint32 GetConverterInfoFlags() const override;
+	bool Convert(FRHICommandListImmediate& RHICmdList, FTextureRHIRef& InDstTexture, const FConversionHints& Hints) override;
+
 	friend class FElectraTextureSampleSupport;
 
-	UE_API FTextureRHIRef InitializeTextureOES(AHardwareBuffer* HardwareBuffer);
-	UE_API FTextureRHIRef InitializeTextureVulkan(AHardwareBuffer* HardwareBuffer);
+	FTextureRHIRef InitializeTextureOES(AHardwareBuffer* HardwareBuffer);
+	FTextureRHIRef InitializeTextureVulkan(AHardwareBuffer* HardwareBuffer);
 
-	UE_API void InitializeTexture(EPixelFormat PixelFormat);
-	UE_API void SetupFromBuffer(const void* InBuffer, int32 InBufferSize);
-	UE_API void SetImageResources(jobject InImageResources);
-	UE_API void CleanupImageResources();
-	UE_API void CopyFromExternalTextureOES(FRHICommandListImmediate& RHICmdList, FTextureRHIRef& InDstTexture, FTextureRHIRef& InSrcTexture, const FVector2f& Scale, const FVector2f& Offset);
-	UE_API void CopyFromExternalTextureVulkan(FRHICommandListImmediate& RHICmdList, FTextureRHIRef& InDstTexture, FTextureRHIRef& InSrcTexture, const FVector2f& InScale, const FVector2f& InOffset);
+	void InitializeTexture(EPixelFormat PixelFormat);
+	void SetupFromBuffer(const void* InBuffer, int32 InBufferSize);
+	void SetImageResources(jobject InImageResources);
+	void CleanupImageResources();
+	void CopyFromExternalTextureOES(FRHICommandListImmediate& RHICmdList, FTextureRHIRef& InDstTexture, FTextureRHIRef& InSrcTexture, const FVector2f& Scale, const FVector2f& Offset);
+	void CopyFromExternalTextureVulkan(FRHICommandListImmediate& RHICmdList, FTextureRHIRef& InDstTexture, FTextureRHIRef& InSrcTexture, const FVector2f& InScale, const FVector2f& InOffset);
 
-	UE_API virtual bool Convert(FRHICommandListImmediate& RHICmdList, FTextureRHIRef& InDstTexture, const FConversionHints& Hints) override;
-	UE_API bool ConvertCpuOutputPath(FRHICommandListImmediate& RHICmdList, FTextureRHIRef& InDstTexture, const FConversionHints& Hints);
-	UE_API bool ConvertGpuOutputPath(FRHICommandListImmediate& RHICmdList, FTextureRHIRef& InDstTexture, const FConversionHints& Hints);
+	bool ConvertCpuOutputPath(FRHICommandListImmediate& RHICmdList, FTextureRHIRef& InDstTexture, const FConversionHints& Hints);
+	bool ConvertGpuOutputPath(FRHICommandListImmediate& RHICmdList, FTextureRHIRef& InDstTexture, const FConversionHints& Hints);
 
-	virtual uint32 GetConverterInfoFlags() const
-	{
-		return ConverterInfoFlags_Default;
-	}
 
 	TSharedPtr<FElectraTextureSampleSupport, ESPMode::ThreadSafe> Support;
 
-	FVideoDecoderOutputAndroid* VideoDecoderOutputAndroid;
-
-	bool UseGpuOutputPath;
-
-	void* Buffer;
-	int32 BufferSize;
-
-	/** Texture resource. */
 	TRefCountPtr<FRHITexture> Texture;
+	TRefCountPtr<FRHIGPUFence> Fence;
 
-	jobject ImageResources;
-	bool bQueuedForConversion;
+	void* Buffer = nullptr;
+	int32 BufferSize = 0;
 
-	TRefCountPtr<FRHIGPUFence> Fence;
+	jobject ImageResources = nullptr;
+	bool bQueuedForConversion = false;
+
+	bool UseGpuOutputPath = false;
 };
 
 
-using FElectraTextureSamplePtr  = TSharedPtr<FElectraTextureSample, ESPMode::ThreadSafe>;
-using FElectraTextureSampleRef  = TSharedRef<FElectraTextureSample, ESPMode::ThreadSafe>;
+using FElectraTextureSamplePtr = TSharedPtr<FElectraTextureSample, ESPMode::ThreadSafe>;
+using FElectraTextureSampleRef = TSharedRef<FElectraTextureSample, ESPMode::ThreadSafe>;
 
 class FElectraTextureSamplePool : public TMediaObjectPool<FElectraTextureSample, FElectraTextureSamplePool>
 {
diff --git a/Engine/Plugins/Media/ElectraUtil/Source/ElectraSamples/Public/Apple/ElectraTextureSample.h b/Engine/Plugins/Media/ElectraUtil/Source/ElectraSamples/Public/Apple/ElectraTextureSample.h
--- a/Engine/Plugins/Media/ElectraUtil/Source/ElectraSamples/Public/Apple/ElectraTextureSample.h
+++ b/Engine/Plugins/Media/ElectraUtil/Source/ElectraSamples/Public/Apple/ElectraTextureSample.h
@@ -7,7 +7,6 @@
 #include "MediaSamples.h"
 #include "IElectraTextureSample.h"
 #include "IMediaTextureSampleConverter.h"
-#include "MediaVideoDecoderOutputApple.h"
 
 #import <AVFoundation/AVFoundation.h>
 #import <VideoToolbox/VideoToolbox.h>
@@ -20,50 +19,34 @@ class ELECTRASAMPLES_API FElectraTextureSample final
     , public IMediaTextureSampleConverter
 {
 public:
-	FElectraTextureSample(const TWeakPtr<FElectraMediaTexConvApple, ESPMode::ThreadSafe> & InTexConv)
-		: Texture(nullptr)
-		, TexConv(InTexConv)
-	{
-	}
+	FElectraTextureSample(const TWeakPtr<FElectraMediaTexConvApple, ESPMode::ThreadSafe>& InTexConv);
+	virtual ~FElectraTextureSample();
 
-	virtual ~FElectraTextureSample()
-	{
-	}
-
-public:
-	void Initialize(FVideoDecoderOutput* InVideoDecoderOutput);
+	bool FinishInitialization() override;
 
 	//~ IMediaTextureSample interface
+	const void* GetBuffer() override;
+	EMediaTextureSampleFormat GetFormat() const override;
+	uint32 GetStride() const override;
+	FRHITexture* GetTexture() const override;
+	IMediaTextureSampleConverter* GetMediaTextureSampleConverter() override;
 
-	virtual const void* GetBuffer() override;
-
-	virtual EMediaTextureSampleFormat GetFormat() const override;
-
-	virtual uint32 GetStride() const override;
-
-	virtual FRHITexture* GetTexture() const override
-	{
-		return Texture;
-	}
-
-	virtual IMediaTextureSampleConverter* GetMediaTextureSampleConverter() override;
-
-private:
-	/** The sample's texture resource. */
-	TRefCountPtr<FRHITexture> Texture;
+	void ShutdownPoolable() override;
 
-	/** Output data from video decoder. */
-	FVideoDecoderOutputApple* VideoDecoderOutputApple;
+	TSharedPtr<TArray64<uint8>, ESPMode::ThreadSafe> Buffer;
+	CVImageBufferRef ImageBufferRef = nullptr;
+	EMediaTextureSampleFormat SampleFormat = EMediaTextureSampleFormat::Undefined;
+	uint32 Stride = 0;
 
 	TWeakPtr<FElectraMediaTexConvApple, ESPMode::ThreadSafe> TexConv;
 
-	virtual uint32 GetConverterInfoFlags() const override;
-    virtual bool Convert(FRHICommandListImmediate& RHICmdList, FTextureRHIRef& InDstTexture, const FConversionHints& Hints) override;
+	uint32 GetConverterInfoFlags() const override;
+    bool Convert(FRHICommandListImmediate& RHICmdList, FTextureRHIRef& InDstTexture, const FConversionHints& Hints) override;
 };
 
 
-using FElectraTextureSamplePtr  = TSharedPtr<FElectraTextureSample, ESPMode::ThreadSafe>;
-using FElectraTextureSampleRef  = TSharedRef<FElectraTextureSample, ESPMode::ThreadSafe>;
+using FElectraTextureSamplePtr = TSharedPtr<FElectraTextureSample, ESPMode::ThreadSafe>;
+using FElectraTextureSampleRef = TSharedRef<FElectraTextureSample, ESPMode::ThreadSafe>;
 
 
 class ELECTRASAMPLES_API FElectraMediaTexConvApple
diff --git a/Engine/Plugins/Media/ElectraUtil/Source/ElectraSamples/Public/IElectraAudioSample.h b/Engine/Plugins/Media/ElectraUtil/Source/ElectraSamples/Public/IElectraAudioSample.h
--- a/Engine/Plugins/Media/ElectraUtil/Source/ElectraSamples/Public/IElectraAudioSample.h
+++ b/Engine/Plugins/Media/ElectraUtil/Source/ElectraSamples/Public/IElectraAudioSample.h
@@ -14,9 +14,17 @@ class FElectraAudioSample
 {
 public:
 	UE_API virtual ~FElectraAudioSample();
-	UE_API bool Initialize(const void* InData, EMediaAudioSampleFormat InFormat, uint32 InNumChannels, uint32 InNumFrames, uint32 InSampleRate, const FMediaTimeStamp& InTime, const FTimespan& InDuration);
+	// Prepare the buffer to receive as many samples as specified for the given format, after which the data can be copied into the buffer.
+	UE_API bool AllocateFor(EMediaAudioSampleFormat InFormat, uint32 InNumChannels, uint32 InNumFrames);
 	uint32 GetAllocatedSize() const
 	{ return NumBytesAllocated;}
+	void* GetWritableBuffer()
+	{ return Buffer; }
+
+	// Sets sample information of the data in the buffer.
+	UE_API void SetParameters(uint32 InSampleRate, const FMediaTimeStamp& InTime, const FTimespan& InDuration);
+	void SetNumFrames(uint32 InNumFrames)
+	{ NumFrames = InNumFrames; }
 
 	// From IMediaAudioSample
 	const void* GetBuffer() override
@@ -35,6 +43,11 @@ public:
 	{ return MediaTimeStamp; }
 	//TOptional<FTimecode> GetTimecode() const override { return TOptional<FTimecode>(); }
 
+
+	DECLARE_DELEGATE_OneParam(FReleaseDelegate, FElectraAudioSample*);
+	FReleaseDelegate& GetReleaseDelegate()
+	{ return ReleaseDelegate; }
+
 	// From IMediaPoolable
 #if !UE_SERVER
 	UE_API void ShutdownPoolable() override;
@@ -49,6 +62,8 @@ protected:
 	uint32 NumFrames = 0;
 	uint32 SampleRate = 0;
 	uint32 NumBytesAllocated = 0;
+	/** Optional delegate to call during ShutdownPoolable(). */
+	FReleaseDelegate ReleaseDelegate;
 };
 
 
diff --git a/Engine/Plugins/Media/ElectraUtil/Source/ElectraSamples/Public/IElectraTextureSample.h b/Engine/Plugins/Media/ElectraUtil/Source/ElectraSamples/Public/IElectraTextureSample.h
--- a/Engine/Plugins/Media/ElectraUtil/Source/ElectraSamples/Public/IElectraTextureSample.h
+++ b/Engine/Plugins/Media/ElectraUtil/Source/ElectraSamples/Public/IElectraTextureSample.h
@@ -3,85 +3,95 @@
 
 #include "IMediaTextureSample.h"
 #include "IMediaTextureSampleConverter.h"
-#include "MediaObjectPool.h"
 #include "Misc/Timespan.h"
 #include "Misc/Timecode.h"
 #include "Misc/Optional.h"
 #include "Templates/SharedPointer.h"
 #include "Templates/RefCounting.h"
-
-#include "MediaVideoDecoderOutput.h"
+#include "MediaObjectPool.h"
+#include "ParameterDictionary.h"
+#include "IElectraTextureSampleEncoding.h"
 
 #define UE_API ELECTRASAMPLES_API
 
+class IVideoDecoderHDRInformation;
+class IVideoDecoderColorimetry;
+class IVideoDecoderTimecode;
+
 class IElectraTextureSampleBase
 	: public IMediaTextureSample
 	, public IMediaPoolable
 {
 public:
-	UE_API void Initialize(FVideoDecoderOutput* InVideoDecoderOutput);
+	UE_API void InitializeCommon(const Electra::FParamDict& InParams);
+	UE_API virtual bool FinishInitialization();
+
+	DECLARE_DELEGATE_OneParam(FReleaseDelegate, IElectraTextureSampleBase*);
+	UE_API FReleaseDelegate& GetReleaseDelegate();
+
 
-	virtual bool IsCacheable() const override
-	{
-		return true;
-	}
+	UE_API bool IsCacheable() const override;
 
 #if !UE_SERVER
-	UE_API virtual void InitializePoolable() override;
-	UE_API virtual void ShutdownPoolable() override;
+	UE_API void InitializePoolable() override;
+	UE_API void ShutdownPoolable() override;
+	UE_API bool IsReadyForReuse() override;
 #endif
 
-	UE_API virtual FIntPoint GetDim() const override;
-	UE_API virtual FIntPoint GetOutputDim() const override;
-
-	UE_API virtual FMediaTimeStamp GetTime() const override;
-	UE_API virtual FTimespan GetDuration() const override;
-
-	virtual TOptional<FTimecode> GetTimecode() const override
-	{
-		return Timecode;
-	}
-	virtual TOptional<FFrameRate> GetFramerate() const override
-	{
-		return Framerate;
-	}
-
-	virtual double GetAspectRatio() const override
-	{
-		return VideoDecoderOutput->GetAspectRatio();
-	}
-
-	virtual EMediaOrientation GetOrientation() const override
-	{
-		return (EMediaOrientation)VideoDecoderOutput->GetOrientation();
-	}
-
-	UE_API virtual bool IsOutputSrgb() const override;
-	UE_API virtual const FMatrix& GetYUVToRGBMatrix() const override;
-	UE_API virtual bool GetFullRange() const override;
-
-	UE_API virtual FMatrix44f GetSampleToRGBMatrix() const override;
-	UE_API virtual const UE::Color::FColorSpace& GetSourceColorSpace() const override;
-	UE_API virtual UE::Color::EEncoding GetEncodingType() const override;
-	UE_API virtual float GetHDRNitsNormalizationFactor() const override;
-	UE_API virtual bool GetDisplayMasteringLuminance(float& OutMin, float& OutMax) const override;
-	UE_API virtual TOptional<UE::Color::FColorSpace> GetDisplayMasteringColorSpace() const override;
-	UE_API virtual bool GetMaxLuminanceLevels(uint16& OutCLL, uint16& OutFALL) const override;
-	UE_API virtual MediaShaders::EToneMapMethod GetToneMapMethod() const override;
-
-	UE_API virtual void SetTime(const FMediaTimeStamp& InTime);
+	UE_API FIntPoint GetDim() const override;
+	UE_API FIntPoint GetOutputDim() const override;
 
-protected:
-	virtual float GetSampleDataScale(bool b10Bit) const { return 1.0f; }
+	UE_API FMediaTimeStamp GetTime() const override;
+	UE_API FTimespan GetDuration() const override;
+
+	UE_API TOptional<FTimecode> GetTimecode() const override;
+	UE_API TOptional<FFrameRate> GetFramerate() const override;
+
+	UE_API double GetAspectRatio() const override;
+	UE_API EMediaOrientation GetOrientation() const override;
 
-	/** Output data from video decoder. */
-	TSharedPtr<FVideoDecoderOutput, ESPMode::ThreadSafe> VideoDecoderOutput;
+	UE_API bool IsOutputSrgb() const override;
+	UE_API const FMatrix& GetYUVToRGBMatrix() const override;
+	UE_API bool GetFullRange() const override;
 
-	/** Quick access for some HDR related info */
-	TWeakPtr<const IVideoDecoderHDRInformation, ESPMode::ThreadSafe> HDRInfo;
-	TWeakPtr<const IVideoDecoderColorimetry, ESPMode::ThreadSafe> Colorimetry;
+	UE_API FMatrix44f GetSampleToRGBMatrix() const override;
+	UE_API const UE::Color::FColorSpace& GetSourceColorSpace() const override;
+	UE_API UE::Color::EEncoding GetEncodingType() const override;
+	UE_API float GetHDRNitsNormalizationFactor() const override;
+	UE_API bool GetDisplayMasteringLuminance(float& OutMin, float& OutMax) const override;
+	UE_API TOptional<UE::Color::FColorSpace> GetDisplayMasteringColorSpace() const override;
+	UE_API bool GetMaxLuminanceLevels(uint16& OutCLL, uint16& OutFALL) const override;
+	UE_API MediaShaders::EToneMapMethod GetToneMapMethod() const override;
+
+	UE_API void SetDim(const FIntPoint& InDim);
+	UE_API virtual float GetSampleDataScale(bool b10Bit) const;
+	UE_API void SetTime(const FMediaTimeStamp& InTime);
+	UE_API void SetDuration(const FTimespan& InDuration);
+	UE_API TSharedPtr<const IVideoDecoderColorimetry, ESPMode::ThreadSafe> GetColorimetry() const;
+	UE_API EPixelFormat GetPixelFormat() const;
+	UE_API EElectraTextureSamplePixelEncoding GetPixelFormatEncoding() const;
+
+protected:
+	FMediaTimeStamp PresentationTimeStamp;
+	FTimespan Duration;
+	FIntPoint ImageDim { 0, 0 };
+	FIntPoint ImageOutputDim { 0, 0 };
+	int32 CropLeft = 0;
+	int32 CropTop = 0;
+	int32 CropRight = 0;
+	int32 CropBottom = 0;
+
+	EPixelFormat PixelFormat = EPixelFormat::PF_Unknown;
+	EElectraTextureSamplePixelEncoding PixelFormatEncoding = EElectraTextureSamplePixelEncoding::Native;
+	float PixelDataScale = 1.0f;
+
+	EMediaOrientation Orientation = EMediaOrientation::Original;
+	double AspectRatio = 1.0;
+
+	/** HDR information */
+	TSharedPtr<const IVideoDecoderHDRInformation, ESPMode::ThreadSafe> HDRInfo;
+	TSharedPtr<const IVideoDecoderColorimetry, ESPMode::ThreadSafe> Colorimetry;
 	/** Optional timecode */
-	TWeakPtr<const IVideoDecoderTimecode> DecoderTimecode;
 	TOptional<FTimecode> Timecode;
 	TOptional<FFrameRate> Framerate;
 
@@ -89,16 +99,20 @@ protected:
 	FMatrix44f SampleToRgbMtx;
 
 	/** YUV to RGB matrix without any adjustments for decoder output specifics */
-	const FMatrix* YuvToRgbMtx;
+	const FMatrix* YuvToRgbMtx = nullptr;
 
 	/** Precomputed colorimetric data */
-	UE::Color::EEncoding ColorEncoding;
+	UE::Color::EEncoding ColorEncoding = UE::Color::EEncoding::None;
 	UE::Color::FColorSpace SourceColorSpace;
 	TOptional<UE::Color::FColorSpace> DisplayMasteringColorSpace;
-	float DisplayMasteringLuminanceMin;
-	float DisplayMasteringLuminanceMax;
-	uint16 MaxCLL;
-	uint16 MaxFALL;
+	float DisplayMasteringLuminanceMin = -1.0f;
+	float DisplayMasteringLuminanceMax = -1.0f;
+	uint16 MaxCLL = 0;
+	uint16 MaxFALL = 0;
+	bool bIsFullRange = false;
+
+	/** Optional delegate to call during ShutdownPoolable(). */
+	FReleaseDelegate ReleaseDelegate;
 };
 
 #undef UE_API
diff --git a/Engine/Plugins/Media/ElectraUtil/Source/ElectraSamples/Public/Linux/ElectraTextureSample.h b/Engine/Plugins/Media/ElectraUtil/Source/ElectraSamples/Public/Linux/ElectraTextureSample.h
--- a/Engine/Plugins/Media/ElectraUtil/Source/ElectraSamples/Public/Linux/ElectraTextureSample.h
+++ b/Engine/Plugins/Media/ElectraUtil/Source/ElectraSamples/Public/Linux/ElectraTextureSample.h
@@ -2,61 +2,40 @@
 
 #pragma once
 
-#include "CoreMinimal.h"
 #include "IElectraTextureSample.h"
-#include "MediaObjectPool.h"
 #include "RHI.h"
-#include "RHIUtilities.h"
+#include "RHIResources.h"
+#include "ShaderCore.h"
 
-#include "Linux/MediaVideoDecoderOutputLinux.h"
 
-class FRHITexture;
 
-class FElectraTextureSampleLinux
+class FElectraTextureSample
 	: public IElectraTextureSampleBase
 {
 public:
-	FElectraTextureSampleLinux() = default;
+	FElectraTextureSample() = default;
+	virtual	~FElectraTextureSample();
 
-	virtual	~FElectraTextureSampleLinux();
-
-public:
-	void Initialize(FVideoDecoderOutput* InVideoDecoderOutput)
-	{
-		IElectraTextureSampleBase::Initialize(InVideoDecoderOutput);
-		VideoDecoderOutputLinux = static_cast<FVideoDecoderOutputLinux*>(InVideoDecoderOutput);
-	}
+	bool FinishInitialization() override;
 
 	const void* GetBuffer() override;
 	uint32 GetStride() const override;
-
 	EMediaTextureSampleFormat GetFormat() const override;
-
-	FRHITexture* GetTexture() const override
-	{
-		return Texture;
-	}
-
-    IMediaTextureSampleConverter* GetMediaTextureSampleConverter() override
-    {
-        return nullptr;
-    }
+	FRHITexture* GetTexture() const override;
+    IMediaTextureSampleConverter* GetMediaTextureSampleConverter() override;
 
 #if !UE_SERVER
 	void ShutdownPoolable() override;
 #endif
 
-private:
-	/** The sample's texture resource. */
-	TRefCountPtr<FRHITexture> Texture;
-
-	/** Output data from video decoder. */
-	FVideoDecoderOutputLinux* VideoDecoderOutputLinux;
+	TSharedPtr<TArray64<uint8>, ESPMode::ThreadSafe> Buffer;
+	EMediaTextureSampleFormat SampleFormat = EMediaTextureSampleFormat::Undefined;
+	uint32 Stride = 0;
+	bool bCanUseSRGB = false;
 };
 
-using FElectraTextureSamplePtr = TSharedPtr<FElectraTextureSampleLinux, ESPMode::ThreadSafe>;
-using FElectraTextureSampleRef = TSharedRef<FElectraTextureSampleLinux, ESPMode::ThreadSafe>;
-using FElectraTextureSample = FElectraTextureSampleLinux;
+using FElectraTextureSamplePtr = TSharedPtr<FElectraTextureSample, ESPMode::ThreadSafe>;
+using FElectraTextureSampleRef = TSharedRef<FElectraTextureSample, ESPMode::ThreadSafe>;
 
 class FElectraTextureSamplePool : public TMediaObjectPool<FElectraTextureSample, FElectraTextureSamplePool>
 {
diff --git a/Engine/Plugins/Media/ElectraUtil/Source/ElectraSamples/Public/Windows/ElectraTextureSample.h b/Engine/Plugins/Media/ElectraUtil/Source/ElectraSamples/Public/Windows/ElectraTextureSample.h
--- a/Engine/Plugins/Media/ElectraUtil/Source/ElectraSamples/Public/Windows/ElectraTextureSample.h
+++ b/Engine/Plugins/Media/ElectraUtil/Source/ElectraSamples/Public/Windows/ElectraTextureSample.h
@@ -6,16 +6,29 @@
 #include "RHIResources.h"
 #include "ShaderCore.h"
 
-#include "MediaVideoDecoderOutputPC.h"
-
 #if !PLATFORM_WINDOWS
 #error "Should only be used on Windows"
 #endif
 
+#include "Windows/AllowWindowsPlatformTypes.h"
+
+#pragma warning(push)
+#pragma warning(disable : 4005) // macro redefinition
+#include <d3d11.h>
+#pragma warning(pop)
+#include <d3d12.h>
+
 THIRD_PARTY_INCLUDES_START
-#include "d3d12.h"
+#include "mfobjects.h"
+#include "mfapi.h"
 THIRD_PARTY_INCLUDES_END
 
+#include "Windows/HideWindowsPlatformTypes.h"
+
+
+
+class FElectraMediaDecoderOutputBufferPool_DX12;
+
 class FElectraTextureSample final
 	: public IElectraTextureSampleBase
 	, public IMediaTextureSampleConverter
@@ -23,56 +36,87 @@ class FElectraTextureSample final
 public:
 
 	/** Default constructor. */
-	FElectraTextureSample()
-		: SampleFormat(EMediaTextureSampleFormat::Undefined)
-		, bCanUseSRGB(false)
-	{ }
+	FElectraTextureSample() = default;
 
-	ELECTRASAMPLES_API void Initialize(FVideoDecoderOutput* VideoDecoderOutput);
-	FVideoDecoderOutput *GetDecoderOutput() { return VideoDecoderOutput.Get(); }
+	ELECTRASAMPLES_API ~FElectraTextureSample();
+
+	ELECTRASAMPLES_API bool FinishInitialization() override;
 
 	//
 	// General Interface
 	//
 
-	ELECTRASAMPLES_API virtual const void* GetBuffer() override;
-	ELECTRASAMPLES_API virtual uint32 GetStride() const override;
-
-	virtual EMediaTextureSampleFormat GetFormat() const override
-	{
-		return SampleFormat;
-	}
+	ELECTRASAMPLES_API const void* GetBuffer() override;
+	ELECTRASAMPLES_API uint32 GetStride() const override;
+	ELECTRASAMPLES_API EMediaTextureSampleFormat GetFormat() const override;
 
 #if WITH_ENGINE
-	ELECTRASAMPLES_API virtual FRHITexture* GetTexture() const override;
+	ELECTRASAMPLES_API FRHITexture* GetTexture() const override;
 #endif //WITH_ENGINE
 
-	ELECTRASAMPLES_API virtual IMediaTextureSampleConverter* GetMediaTextureSampleConverter() override;
+	ELECTRASAMPLES_API IMediaTextureSampleConverter* GetMediaTextureSampleConverter() override;
 
 #if !UE_SERVER
-	ELECTRASAMPLES_API virtual void ShutdownPoolable() override;
+	ELECTRASAMPLES_API bool IsReadyForReuse() override;
+	ELECTRASAMPLES_API void ShutdownPoolable() override;
 #endif
 
-private:
-	ELECTRASAMPLES_API virtual float GetSampleDataScale(bool b10Bit) const override;
+//private:
+	ELECTRASAMPLES_API float GetSampleDataScale(bool b10Bit) const override;
 
-	ELECTRASAMPLES_API virtual bool Convert(FRHICommandListImmediate& RHICmdList, FTextureRHIRef& InDstTexture, const FConversionHints& Hints) override;
+	ELECTRASAMPLES_API bool Convert(FRHICommandListImmediate& RHICmdList, FTextureRHIRef& InDstTexture, const FConversionHints& Hints) override;
 	virtual uint32 GetConverterInfoFlags() const
 	{
 		return ConverterInfoFlags_PreprocessOnly;
 	}
 
+	ELECTRASAMPLES_API TRefCountPtr<IUnknown> GetSync(uint64& OutSyncValue);
+	ELECTRASAMPLES_API TRefCountPtr<IUnknown> GetTextureD3D() const;
+
+
 	/** The sample format. */
-	EMediaTextureSampleFormat SampleFormat;
+	EMediaTextureSampleFormat SampleFormat = EMediaTextureSampleFormat::Undefined;
 
 	/** Destination Texture resource (from Rendering device) */
 	mutable FTextureRHIRef Texture;
 
 	/** True if texture format could support sRGB conversion in HW */
-	bool bCanUseSRGB;
+	bool bCanUseSRGB = false;
 
-	/** Output data from video decoder. Baseclass holds reference */
-	FVideoDecoderOutputPC* VideoDecoderOutputPC;
+	enum class ESourceType
+	{
+		Unknown,
+		Buffer,
+		SharedTextureDX11,
+		ResourceDX12
+	};
+
+	ESourceType SourceType = ESourceType::Unknown;
+
+	// DX11 related members (output texture with device that created it)
+	TRefCountPtr<ID3D11Texture2D> TextureDX11;
+	TRefCountPtr<ID3D11Device> D3D11Device;
+
+	// DX12 related members
+	mutable TRefCountPtr<ID3D12Resource> TextureDX12;
+	FIntPoint TextureDX12Dim {0, 0};
+	TRefCountPtr<ID3D12Fence> D3DFence;
+	uint64 FenceValue = 0;
+	TRefCountPtr<ID3D12Resource> DecoderOutputResource;
+	TRefCountPtr<ID3D12CommandAllocator> D3DCmdAllocator;
+	TRefCountPtr<ID3D12GraphicsCommandList> D3DCmdList;
+	TSharedPtr<FElectraMediaDecoderOutputBufferPool_DX12> D3D12ResourcePool;
+
+	// CPU-side buffer
+	TSharedPtr<TArray64<uint8>, ESPMode::ThreadSafe> Buffer;
+	uint32 Stride = 0;
+
+	// Dimension of any internally allocated buffer - stored explicitly to cover various special cases
+	FIntPoint SampleDim {0, 0};
+
+	ELECTRASAMPLES_API void ClearDX11Vars();
+	ELECTRASAMPLES_API void ClearDX12Vars();
+	ELECTRASAMPLES_API void ClearBufferVars();
 };
 
 using FElectraTextureSamplePtr = TSharedPtr<FElectraTextureSample, ESPMode::ThreadSafe>;
