Reworked ElectraSamples, ElectraUtils and the decoder output of Electra Player

diff --git a/Engine/Plugins/Media/ElectraUtil/Source/ElectraSamples/Public/Windows/WindowsElectraTextureSampleGPUBufferHelper.h b/Engine/Plugins/Media/ElectraUtil/Source/ElectraSamples/Public/Windows/WindowsElectraTextureSampleGPUBufferHelper.h
new file mode 100644
--- /dev/null
+++ b/Engine/Plugins/Media/ElectraUtil/Source/ElectraSamples/Public/Windows/WindowsElectraTextureSampleGPUBufferHelper.h
@@ -0,0 +1,550 @@
+// Copyright Epic Games, Inc. All Rights Reserved.
+
+#pragma once
+
+#include "CoreMinimal.h"
+
+/*********************************************************************************************************************/
+/*********************************************************************************************************************/
+/*********************************************************************************************************************/
+
+#include "Windows/AllowWindowsPlatformTypes.h"
+THIRD_PARTY_INCLUDES_START
+#include "mfobjects.h"
+#include "d3d12.h"
+#include "d3dx12.h"
+THIRD_PARTY_INCLUDES_END
+#include "Windows/HideWindowsPlatformTypes.h"
+
+#define	ELECTRA_MEDIAGPUBUFFER_DX12 1
+
+/*********************************************************************************************************************/
+/*********************************************************************************************************************/
+/*********************************************************************************************************************/
+
+#if defined(_WIN64) && defined(__SSE2__)
+
+#define ALLOW_USE_OF_SSEMEMCPY 1
+
+#include <intrin.h>
+
+static void SSE2MemCpy(const void* InDst, const void* InSrc, unsigned int InSize)
+{
+	__m128i* Dst = (__m128i*)InDst;
+	__m128i* Src = (__m128i*)InSrc;
+
+	InSize = InSize >> 7;
+	for(unsigned int i = 0; i < InSize; i++)
+	{
+		_mm_prefetch((char*)(Src + 8), 0);
+		_mm_prefetch((char*)(Src + 10), 0);
+		_mm_prefetch((char*)(Src + 12), 0);
+		_mm_prefetch((char*)(Src + 14), 0);
+		__m128i m0 = _mm_load_si128(Src + 0);
+		__m128i m1 = _mm_load_si128(Src + 1);
+		__m128i m2 = _mm_load_si128(Src + 2);
+		__m128i m3 = _mm_load_si128(Src + 3);
+		__m128i m4 = _mm_load_si128(Src + 4);
+		__m128i m5 = _mm_load_si128(Src + 5);
+		__m128i m6 = _mm_load_si128(Src + 6);
+		__m128i m7 = _mm_load_si128(Src + 7);
+		_mm_stream_si128(Dst + 0, m0);
+		_mm_stream_si128(Dst + 1, m1);
+		_mm_stream_si128(Dst + 2, m2);
+		_mm_stream_si128(Dst + 3, m3);
+		_mm_stream_si128(Dst + 4, m4);
+		_mm_stream_si128(Dst + 5, m5);
+		_mm_stream_si128(Dst + 6, m6);
+		_mm_stream_si128(Dst + 7, m7);
+		Src += 8;
+		Dst += 8;
+	}
+}
+
+#else
+#define ALLOW_USE_OF_SSEMEMCPY 0
+#endif // _WIN64
+
+/*********************************************************************************************************************/
+/*********************************************************************************************************************/
+/*********************************************************************************************************************/
+
+class FElectraMediaDecoderOutputBufferPool_DX12;
+
+// Simple DX12 heap and fence manager for use with decoder output buffers (both upload heaps or default GPU heaps)
+class FElectraMediaDecoderOutputBufferPoolBlock_DX12
+{
+	struct FResourceInfo
+	{
+		FResourceInfo(TWeakPtr<FElectraMediaDecoderOutputBufferPool_DX12> InPool, uint32 InBlockIdx, uint32 InBufferIdx, TRefCountPtr<ID3D12Heap> InHeap) : Pool(InPool), BlockIdx(InBlockIdx), BufferIdx(InBufferIdx), Heap(InHeap) {}
+
+		TWeakPtr<FElectraMediaDecoderOutputBufferPool_DX12> Pool;
+		uint32 BlockIdx;
+		uint32 BufferIdx;
+		TRefCountPtr<ID3D12Heap> Heap;
+	};
+
+public:
+	// Create instance for use with buffers (usually for uploading data)
+	FElectraMediaDecoderOutputBufferPoolBlock_DX12(uint32 InBlockIdx, TRefCountPtr<ID3D12Device> InD3D12Device, uint32 InMaxNumBuffers, uint32 Width, uint32 Height, uint32 BytesPerPixel, D3D12_HEAP_TYPE InD3D12HeapType = D3D12_HEAP_TYPE_UPLOAD)
+		: D3D12Device(InD3D12Device)
+		, D3D12HeapType(InD3D12HeapType)
+		, MaxNumBuffers(InMaxNumBuffers)
+		, FreeMask((uint32)((1ULL << InMaxNumBuffers) - 1))
+		, BlockIdx(InBlockIdx)
+	{
+		check(InMaxNumBuffers <= 32);
+
+		// Get the size needed per buffer
+		BufferPitch = Align(Width * BytesPerPixel, D3D12_TEXTURE_DATA_PITCH_ALIGNMENT);
+		BufferSize = Align(BufferPitch * Height, FMath::Max(D3D12_TEXTURE_DATA_PLACEMENT_ALIGNMENT, D3D12_DEFAULT_RESOURCE_PLACEMENT_ALIGNMENT));
+
+		InitCommon();
+	}
+
+	// Create instance for use with textures (usually to receive any upload buffers or internal decoder texture data)
+	FElectraMediaDecoderOutputBufferPoolBlock_DX12(uint32 InBlockIdx, TRefCountPtr<ID3D12Device> InD3D12Device, uint32 InMaxNumBuffers, uint32 Width, uint32 Height, DXGI_FORMAT PixFmt, D3D12_HEAP_TYPE InD3D12HeapType = D3D12_HEAP_TYPE_DEFAULT)
+		: D3D12Device(InD3D12Device)
+		, D3D12HeapType(InD3D12HeapType)
+		, MaxNumBuffers(InMaxNumBuffers)
+		, FreeMask((uint32)((1ULL << InMaxNumBuffers) - 1))
+		, BlockIdx(InBlockIdx)
+	{
+		check(InMaxNumBuffers <= 32);
+
+		// Get the size needed per buffer / texture
+		D3D12_RESOURCE_DESC Desc = {};
+		Desc.MipLevels = 1;
+		Desc.Format = PixFmt;
+		Desc.Width = Width;
+		Desc.Height = Height;
+		Desc.Flags = D3D12_RESOURCE_FLAG_NONE;
+		Desc.DepthOrArraySize = 1;
+		Desc.SampleDesc.Count = 1;
+		Desc.SampleDesc.Quality = 0;
+		Desc.Layout = D3D12_TEXTURE_LAYOUT_UNKNOWN;
+		Desc.Dimension = D3D12_RESOURCE_DIMENSION_TEXTURE2D;
+		D3D12_RESOURCE_ALLOCATION_INFO AllocInfo = InD3D12Device->GetResourceAllocationInfo(0, 1, &Desc);
+
+		BufferPitch = 0;
+		BufferSize = Align(AllocInfo.SizeInBytes, AllocInfo.Alignment);
+		InitCommon();
+	}
+
+	~FElectraMediaDecoderOutputBufferPoolBlock_DX12()
+	{
+	}
+
+	uint32 GetMaxNumBuffers() const
+	{
+		return MaxNumBuffers;
+	}
+
+	// Check if the current setup is compatible with the new parameters
+	bool IsCompatibleAsBuffer(uint32 Width, uint32 Height, uint32 BytesPerPixel) const
+	{
+		uint32 NewBufferPitch = Align(Width * BytesPerPixel, D3D12_TEXTURE_DATA_PITCH_ALIGNMENT);
+		uint32 NewBufferSize = Align(NewBufferPitch * Height, FMath::Max(D3D12_TEXTURE_DATA_PLACEMENT_ALIGNMENT, D3D12_DEFAULT_RESOURCE_PLACEMENT_ALIGNMENT));
+		return BufferSize >= NewBufferSize && BufferPitch == NewBufferPitch;
+	}
+
+	// Check if the current setup is compatible with the new parameters
+	bool IsCompatibleAsTexture(uint32 Width, uint32 Height, DXGI_FORMAT PixFmt) const
+	{
+		// Get the size needed per buffer / texture
+		D3D12_RESOURCE_DESC Desc = {};
+		Desc.MipLevels = 1;
+		Desc.Format = PixFmt;
+		Desc.Width = Width;
+		Desc.Height = Height;
+		Desc.Flags = D3D12_RESOURCE_FLAG_NONE;
+		Desc.DepthOrArraySize = 1;
+		Desc.SampleDesc.Count = 1;
+		Desc.SampleDesc.Quality = 0;
+		Desc.Layout = D3D12_TEXTURE_LAYOUT_UNKNOWN;
+		Desc.Dimension = D3D12_RESOURCE_DIMENSION_TEXTURE2D;
+		D3D12_RESOURCE_ALLOCATION_INFO AllocInfo = D3D12Device->GetResourceAllocationInfo(0, 1, &Desc);
+
+		uint32 NewBufferSize = Align(AllocInfo.SizeInBytes, AllocInfo.Alignment);
+
+		return BufferSize >= NewBufferSize;
+	}
+
+	// Check if a buffer is available
+	bool BufferAvailable() const
+	{
+		return FPlatformAtomics::AtomicRead((const int32*)&FreeMask) != 0;
+	}
+
+	// Allocate a buffer resource and return suitable sync fence data
+	TRefCountPtr<ID3D12Resource> AllocateOutputDataAsBuffer(TWeakPtr<FElectraMediaDecoderOutputBufferPool_DX12> InPoolWeakRef, uint32& OutPitch)
+	{
+		TRefCountPtr<ID3D12Resource> Resource = AllocateBuffer(InPoolWeakRef, OutPitch, [NumBytes = BufferSize](D3D12_RESOURCE_DESC& Desc)
+		{
+			Desc.MipLevels = 1;
+			Desc.Format = DXGI_FORMAT_UNKNOWN;
+			Desc.Width = NumBytes;
+			Desc.Height = 1;
+			Desc.Flags = D3D12_RESOURCE_FLAG_NONE;
+			Desc.DepthOrArraySize = 1;
+			Desc.SampleDesc.Count = 1;
+			Desc.SampleDesc.Quality = 0;
+			Desc.Layout = D3D12_TEXTURE_LAYOUT_ROW_MAJOR;
+			Desc.Dimension = D3D12_RESOURCE_DIMENSION_BUFFER;
+		});
+		return Resource;
+	}
+
+	// Allocate a texture resource and return suitable sync fence data
+	TRefCountPtr<ID3D12Resource> AllocateOutputDataAsTexture(TWeakPtr<FElectraMediaDecoderOutputBufferPool_DX12> InPoolWeakRef, uint32 InWidth, uint32 InHeight, DXGI_FORMAT InPixFmt)
+	{
+		uint32 OutPitch;
+		TRefCountPtr<ID3D12Resource> Resource = AllocateBuffer(InPoolWeakRef, OutPitch, [InWidth, InHeight, InPixFmt](D3D12_RESOURCE_DESC& Desc)
+		{
+			Desc.MipLevels = 1;
+			Desc.Format = InPixFmt;
+			Desc.Width = InWidth;
+			Desc.Height = InHeight;
+			Desc.Flags = D3D12_RESOURCE_FLAG_NONE;
+			Desc.DepthOrArraySize = 1;
+			Desc.SampleDesc.Count = 1;
+			Desc.SampleDesc.Quality = 0;
+			Desc.Layout = D3D12_TEXTURE_LAYOUT_UNKNOWN;
+			Desc.Dimension = D3D12_RESOURCE_DIMENSION_TEXTURE2D;
+		});
+		return Resource;
+	}
+
+private:
+	void InitCommon()
+	{
+		D3D12_HEAP_DESC HeapDesc = {};
+		HeapDesc.Alignment = 0;
+		HeapDesc.SizeInBytes = BufferSize * MaxNumBuffers;
+		HeapDesc.Properties = CD3DX12_HEAP_PROPERTIES(D3D12HeapType);
+		HeapDesc.Flags = D3D12_HEAP_FLAG_CREATE_NOT_ZEROED | ((D3D12HeapType == D3D12_HEAP_TYPE_UPLOAD) ? D3D12_HEAP_FLAG_ALLOW_ONLY_BUFFERS : D3D12_HEAP_FLAG_ALLOW_ONLY_NON_RT_DS_TEXTURES);
+		HRESULT Res = D3D12Device->CreateHeap(&HeapDesc, IID_PPV_ARGS(D3D12OutputHeap.GetInitReference()));
+		check(SUCCEEDED(Res));
+		if (Res != S_OK)
+		{
+			FreeMask = 0;
+		}
+#if !UE_BUILD_SHIPPING
+		if (D3D12OutputHeap)
+		{
+			D3D12OutputHeap->SetName(TEXT("ElectraOutputBufferPoolHeap"));
+		}
+#endif
+	}
+
+	TRefCountPtr<ID3D12Resource> AllocateBuffer(TWeakPtr<FElectraMediaDecoderOutputBufferPool_DX12> InPoolWeakRef, uint32& OutBufferPitch, TFunction<void(D3D12_RESOURCE_DESC& Desc)> && InitializeDesc)
+	{
+		TRefCountPtr<ID3D12Resource> Resource;
+
+		// Find the buffer index we can use...
+		int32 BufferIdx, OldFreeMask, NewFreeMask;
+		do
+		{
+			OldFreeMask = (int32)FreeMask;
+			if (OldFreeMask == 0)
+			{
+				BufferIdx = -1;
+				break;
+			}
+			BufferIdx = 31 - FMath::CountLeadingZeros(OldFreeMask);
+			NewFreeMask = OldFreeMask & ~(1 << BufferIdx);
+		} while (OldFreeMask != FPlatformAtomics::InterlockedCompareExchange((int32*)&FreeMask, NewFreeMask, OldFreeMask));
+
+		// Anything?
+		if (BufferIdx >= 0)
+		{
+			// Make a placed resource to use...
+			D3D12_RESOURCE_DESC Desc = {};
+			InitializeDesc(Desc);
+			HRESULT Res = D3D12Device->CreatePlacedResource(D3D12OutputHeap, BufferSize * BufferIdx, &Desc, (D3D12HeapType == D3D12_HEAP_TYPE_UPLOAD) ? D3D12_RESOURCE_STATE_GENERIC_READ : D3D12_RESOURCE_STATE_COMMON, nullptr, IID_PPV_ARGS(Resource.GetInitReference()));
+			check(SUCCEEDED(Res));
+			if (Res == S_OK)
+			{
+#if !UE_BUILD_SHIPPING
+				Resource->SetName(TEXT("ElectraOutputBufferPoolBufferResource"));
+#endif
+
+				OutBufferPitch = BufferPitch;
+
+				// Register a destruction callback so we can reset our memory management here transparently...
+				TRefCountPtr<ID3DDestructionNotifier> Notifier;
+				Res = Resource->QueryInterface(__uuidof(ID3DDestructionNotifier), (void**)Notifier.GetInitReference());
+				check(SUCCEEDED(Res));
+
+				// note: we keep a reference to the heap in our context data for the destruction callback, so we can ensure the heap lives as long as there are placed resources in it
+				UINT CallbackID;
+				Res = Notifier->RegisterDestructionCallback(ResourceDestructionCallback, new FResourceInfo(InPoolWeakRef, BlockIdx, BufferIdx, D3D12OutputHeap), &CallbackID);
+				check(SUCCEEDED(Res));
+			}
+			else
+			{
+				FreeBuffer(BufferIdx);
+			}
+		}
+
+		return Resource;
+	}
+
+	void FreeBuffer(int32 InBufferIdx)
+	{
+		// Reset buffer index to free...
+		int32 OldFreeMask, NewFreeMask;
+		do
+		{
+			OldFreeMask = (int32)FreeMask;
+			NewFreeMask = OldFreeMask | (1 << InBufferIdx);
+		} while (OldFreeMask != FPlatformAtomics::InterlockedCompareExchange((int32*)&FreeMask, NewFreeMask, OldFreeMask));
+	}
+
+	static void ResourceDestructionCallback(void* Context);
+
+	TRefCountPtr<ID3D12Device> D3D12Device;
+	D3D12_HEAP_TYPE D3D12HeapType;
+	TRefCountPtr<ID3D12Heap> D3D12OutputHeap;
+	uint32 MaxNumBuffers;
+	uint32 BufferSize;
+	uint32 BufferPitch;
+	uint32 FreeMask;
+	uint32 BlockIdx;
+};
+
+
+class FElectraMediaDecoderOutputBufferPool_DX12 : public TSharedFromThis<FElectraMediaDecoderOutputBufferPool_DX12, ESPMode::ThreadSafe>
+{
+public:
+	// Create instance for use with buffers (usually for uploading data)
+	FElectraMediaDecoderOutputBufferPool_DX12(TRefCountPtr<ID3D12Device> InD3D12Device, uint32 InMaxNumBuffers, uint32 InWidth, uint32 InHeight, uint32 InBytesPerPixel, D3D12_HEAP_TYPE InD3D12HeapType = D3D12_HEAP_TYPE_UPLOAD)
+		: D3D12Device(InD3D12Device)
+		, LastFenceValue(0)
+		, D3D12HeapType(InD3D12HeapType)
+		, Width(InWidth)
+		, Height(InHeight)
+		, PixFmt(DXGI_FORMAT_UNKNOWN)
+		, BytesPerPixel(InBytesPerPixel)
+		, MaxNumBuffers(InMaxNumBuffers)
+	{
+		Blocks.Emplace(new FElectraMediaDecoderOutputBufferPoolBlock_DX12(0, D3D12Device, MaxNumBuffers < 32 ? MaxNumBuffers : 32, Width, Height, BytesPerPixel, D3D12HeapType));
+		InitCommon();
+	}
+
+	// Create instance for use with textures (usually to receive any upload buffers or internal decoder texture data)
+	FElectraMediaDecoderOutputBufferPool_DX12(TRefCountPtr<ID3D12Device> InD3D12Device, uint32 InMaxNumBuffers, uint32 InWidth, uint32 InHeight, DXGI_FORMAT InPixFmt, D3D12_HEAP_TYPE InD3D12HeapType = D3D12_HEAP_TYPE_DEFAULT)
+		: D3D12Device(InD3D12Device)
+		, LastFenceValue(0)
+		, D3D12HeapType(InD3D12HeapType)
+		, Width(InWidth)
+		, Height(InHeight)
+		, PixFmt(InPixFmt)
+		, BytesPerPixel(0)
+		, MaxNumBuffers(InMaxNumBuffers)
+	{
+		Blocks.Emplace(new FElectraMediaDecoderOutputBufferPoolBlock_DX12(0, D3D12Device, MaxNumBuffers < 32 ? MaxNumBuffers : 32, Width, Height, PixFmt, D3D12HeapType));
+		InitCommon();
+	}
+
+	~FElectraMediaDecoderOutputBufferPool_DX12()
+	{
+		for(FElectraMediaDecoderOutputBufferPoolBlock_DX12* Block : Blocks)
+		{
+			delete Block;
+		}
+	}
+
+	// Check if the current setup is compatible with the new parameters
+	bool IsCompatibleAsBuffer(uint32 InMaxNumBuffers, uint32 InWidth, uint32 InHeight, uint32 InBytesPerPixel) const
+	{
+		return InMaxNumBuffers <= MaxNumBuffers ? Blocks[0]->IsCompatibleAsBuffer(InWidth, InHeight, InBytesPerPixel) : false;
+	}
+
+	bool IsCompatibleAsTexture(uint32 InMaxNumBuffers, uint32 InWidth, uint32 InHeight, DXGI_FORMAT InPixFmt) const
+	{
+		return InMaxNumBuffers <= MaxNumBuffers ? Blocks[0]->IsCompatibleAsTexture(InWidth, InHeight, InPixFmt) : false;
+	}
+
+	// Check if a buffer is available
+	bool BufferAvailable() const
+	{
+		// Check if any already existing block has an available entry
+		for(FElectraMediaDecoderOutputBufferPoolBlock_DX12* Block : Blocks)
+		{
+			if (Block->BufferAvailable())
+			{
+				return true;
+			}
+		}
+		// Can we still create a new entry?
+		const uint32 NumAllocatedSlots = ((Blocks.Num() - 1) * 32U) + Blocks.Last()->GetMaxNumBuffers();
+		return NumAllocatedSlots < MaxNumBuffers;
+	}
+
+	class FOutputData
+	{
+	public:
+		~FOutputData()
+		{
+			check(ReadyToDestroy());
+		}
+
+		TRefCountPtr<ID3D12Resource> Resource;
+		TRefCountPtr<ID3D12Fence> Fence;
+		uint64 FenceValue;
+
+		bool ReadyToDestroy() const
+		{
+			bool bOk = true;
+#if !UE_BUILD_SHIPPING
+			if (Resource.IsValid())
+			{
+				bOk = Resource->AddRef() > 1 || Fence->GetCompletedValue() >= FenceValue;
+				Resource->Release();
+			}
+#endif
+			return bOk;
+		}
+	};
+
+	bool AllocateOutputDataAsBuffer(FOutputData& OutData, uint32& OutPitch)
+	{
+		check(BytesPerPixel != 0);
+
+		OutData.Resource = nullptr;
+		for(FElectraMediaDecoderOutputBufferPoolBlock_DX12* Block : Blocks)
+		{
+			if (Block->BufferAvailable())
+			{
+				OutData.Resource = Block->AllocateOutputDataAsBuffer(AsWeak(), OutPitch);
+				if (OutData.Resource)
+				{
+					break;
+				}
+			}
+		}
+
+		const uint32 NumAllocatedSlots = ((Blocks.Num() - 1) * 32U) + Blocks.Last()->GetMaxNumBuffers();
+		if (!OutData.Resource.IsValid() && NumAllocatedSlots < MaxNumBuffers)
+		{
+			Blocks.Emplace(new FElectraMediaDecoderOutputBufferPoolBlock_DX12(Blocks.Num(), D3D12Device, MaxNumBuffers-NumAllocatedSlots > 32 ? 32 : MaxNumBuffers-NumAllocatedSlots, Width, Height, BytesPerPixel, D3D12HeapType));
+			OutData.Resource = Blocks.Last()->AllocateOutputDataAsBuffer(AsWeak(), OutPitch);
+		}
+
+		if (OutData.Resource.IsValid())
+		{
+			OutData.Fence = GetUpdatedBufferFence(OutData.FenceValue);
+			return true;
+		}
+		return false;
+	}
+
+	bool AllocateOutputDataAsTexture(FOutputData& OutData, uint32 InWidth, uint32 InHeight, DXGI_FORMAT InPixFmt)
+	{
+		check(PixFmt != DXGI_FORMAT_UNKNOWN);
+
+		OutData.Resource = nullptr;
+		for(FElectraMediaDecoderOutputBufferPoolBlock_DX12* Block : Blocks)
+		{
+			if (Block->BufferAvailable())
+			{
+				OutData.Resource = Block->AllocateOutputDataAsTexture(AsWeak(), InWidth, InHeight, InPixFmt);
+				if (OutData.Resource)
+				{
+					break;
+				}
+			}
+		}
+
+		const uint32 NumAllocatedSlots = ((Blocks.Num() - 1) * 32U) + Blocks.Last()->GetMaxNumBuffers();
+		if (!OutData.Resource.IsValid() && NumAllocatedSlots < MaxNumBuffers)
+		{
+			Blocks.Emplace(new FElectraMediaDecoderOutputBufferPoolBlock_DX12(Blocks.Num(), D3D12Device, MaxNumBuffers-NumAllocatedSlots > 32 ? 32 : MaxNumBuffers-NumAllocatedSlots, Width, Height, PixFmt, D3D12HeapType));
+			OutData.Resource = Blocks.Last()->AllocateOutputDataAsTexture(AsWeak(), InWidth, InHeight, InPixFmt);
+		}
+
+		if (OutData.Resource.IsValid())
+		{
+			OutData.Fence = GetUpdatedBufferFence(OutData.FenceValue);
+			return true;
+		}
+		return false;
+	}
+
+	TRefCountPtr<ID3D12Fence> GetUpdatedBufferFence(uint64& FenceValue)
+	{
+		FenceValue = ++LastFenceValue;
+		return D3D12BufferFence;
+	}
+
+	// Helper function to copy simple, linear texture data between differently picthed buffers
+	static void CopyWithPitchAdjust(uint8* Dst, uint32 DstPitch, const uint8* Src, uint32 SrcPitch, uint32 NumRows)
+	{
+		if (DstPitch != SrcPitch)
+		{
+			for(uint32 Y = NumRows; Y > 0; --Y)
+			{
+				FMemory::Memcpy(Dst, Src, SrcPitch);
+				Src += SrcPitch;
+				Dst += DstPitch;
+			}
+		}
+		else
+		{
+			uint32 NumBytes = SrcPitch * NumRows;
+#if ALLOW_USE_OF_SSEMEMCPY
+			if ((NumBytes & 0x7f) == 0)
+			{
+				SSE2MemCpy(Dst, Src, NumBytes);
+			}
+			else
+#endif
+			{
+				FMemory::Memcpy(Dst, Src, NumBytes);
+			}
+		}
+	}
+
+private:
+	friend class FElectraMediaDecoderOutputBufferPoolBlock_DX12;
+
+	void InitCommon()
+	{
+		HRESULT Res = D3D12Device->CreateFence(0, D3D12_FENCE_FLAG_NONE, IID_PPV_ARGS(D3D12BufferFence.GetInitReference()));
+		check(SUCCEEDED(Res));
+#if !UE_BUILD_SHIPPING
+		D3D12BufferFence->SetName(TEXT("ElectraOutputBufferPoolFence"));
+#endif
+	}
+
+	TArray<FElectraMediaDecoderOutputBufferPoolBlock_DX12*> Blocks;
+
+	TRefCountPtr<ID3D12Device> D3D12Device;
+	TRefCountPtr<ID3D12Fence> D3D12BufferFence;
+	uint64 LastFenceValue;
+	D3D12_HEAP_TYPE D3D12HeapType;
+	uint32 Width;
+	uint32 Height;
+	DXGI_FORMAT PixFmt;
+	uint32 BytesPerPixel;
+	uint32 MaxNumBuffers;
+	};
+
+
+inline void FElectraMediaDecoderOutputBufferPoolBlock_DX12::ResourceDestructionCallback(void* Context)
+{
+	auto ResourceInfo = reinterpret_cast<FResourceInfo*>(Context);
+
+	if (auto Pool = ResourceInfo->Pool.Pin())
+	{
+		Pool->Blocks[ResourceInfo->BlockIdx]->FreeBuffer(ResourceInfo->BufferIdx);
+	}
+
+	// Get rid of our resource tracking info block...
+	// (this will also release the ref to the heap the resource was on)
+	delete ResourceInfo;
+}
+
+
+#include "CommonElectraTextureSampleGPUBufferHelper.h"
diff --git a/Engine/Plugins/Media/FlacDecoderElectra/Source/FlacDecoderElectra/FlacDecoderElectra.Build.cs b/Engine/Plugins/Media/FlacDecoderElectra/Source/FlacDecoderElectra/FlacDecoderElectra.Build.cs
--- a/Engine/Plugins/Media/FlacDecoderElectra/Source/FlacDecoderElectra/FlacDecoderElectra.Build.cs
+++ b/Engine/Plugins/Media/FlacDecoderElectra/Source/FlacDecoderElectra/FlacDecoderElectra.Build.cs
@@ -9,6 +9,7 @@ namespace UnrealBuildTool.Rules
 			PrivateDependencyModuleNames.AddRange(
 				new string[] {
 					"Core",
+                    "ElectraBase",
 					"ElectraCodecFactory",
 					"ElectraDecoders",
 					"libFLAC"
diff --git a/Engine/Plugins/Media/FlacDecoderElectra/Source/FlacDecoderElectra/Private/FlacDecoder/ElectraMediaFlacDecoder.cpp b/Engine/Plugins/Media/FlacDecoderElectra/Source/FlacDecoderElectra/Private/FlacDecoder/ElectraMediaFlacDecoder.cpp
--- a/Engine/Plugins/Media/FlacDecoderElectra/Source/FlacDecoderElectra/Private/FlacDecoder/ElectraMediaFlacDecoder.cpp
+++ b/Engine/Plugins/Media/FlacDecoderElectra/Source/FlacDecoderElectra/Private/FlacDecoder/ElectraMediaFlacDecoder.cpp
@@ -13,9 +13,9 @@
 #include "IElectraDecoder.h"
 #include "IElectraDecoderOutputAudio.h"
 #include "ElectraDecodersUtils.h"
-#include "Utils/MPEG/ElectraUtilsMP4.h"		// for parsing the 'dfLa' box
+#include "Utilities/UtilitiesMP4.h"		// for parsing the 'dfLa' box
+#include "Utilities/ElectraBitstream.h"
 #include "Utils/ElectraBitstreamProcessorDefault.h"
-#include "Utils/ElectraBitstreamReader.h"
 
 #include "stream_decoder.h"
 
@@ -311,7 +311,7 @@ bool FElectraFlacDecoder::Parse_dfLa(FdfLaConfig& OutConfig, const TArray<uint8>
 		return !bFailOnError ? false : PostError(0, TEXT("Incomplete 'dfLa' box"), ERRCODE_INTERNAL_FAILED_TO_PARSE_CSD);
 	}
 
-	ElectraDecodersUtil::FMP4AtomReader rd(IndfLaBox.GetData(), IndfLaBox.Num());
+	Electra::UtilitiesMP4::FMP4AtomReader rd(IndfLaBox);
 	uint8 Value8 = 0;
 	rd.Read(Value8);		// Version
 	if (Value8 != 0)
@@ -355,7 +355,7 @@ bool FElectraFlacDecoder::Parse_dfLa(FdfLaConfig& OutConfig, const TArray<uint8>
 		// STREAMINFO block?
 		if (BlockType == 0)
 		{
-			ElectraDecodersUtil::FElectraBitstreamReader si(BlockData.GetData(), BlockData.Num());
+			Electra::FBitstreamReader si(BlockData.GetData(), BlockData.Num());
 			uint16 MinimumBlockSize = si.GetBits(16);
 			uint16 MaximumBlockSize = si.GetBits(16);
 			uint32 MinimumFrameSize = si.GetBits(24);
diff --git a/Engine/Plugins/Media/HAPDecoderElectra/Source/HAPDecoderElectra/Private/HAPDecoder/ElectraMediaHAPDecoder.cpp b/Engine/Plugins/Media/HAPDecoderElectra/Source/HAPDecoderElectra/Private/HAPDecoder/ElectraMediaHAPDecoder.cpp
--- a/Engine/Plugins/Media/HAPDecoderElectra/Source/HAPDecoderElectra/Private/HAPDecoder/ElectraMediaHAPDecoder.cpp
+++ b/Engine/Plugins/Media/HAPDecoderElectra/Source/HAPDecoderElectra/Private/HAPDecoder/ElectraMediaHAPDecoder.cpp
@@ -16,7 +16,7 @@
 #include "Utils/ElectraBitstreamProcessorDefault.h"
 
 #include "ElectraDecodersPlatformResources.h"
-#include COMPILED_PLATFORM_HEADER(ElectraDecoderGPUBufferHelpers.h)
+#include COMPILED_PLATFORM_HEADER(ElectraTextureSampleGPUBufferHelper.h)
 
 #include "hap.h"
 
@@ -118,7 +118,7 @@ public:
 	{
 		return NumBuffers;
 	}
-	TSharedPtr<TArray<uint8>, ESPMode::ThreadSafe> GetBufferDataByIndex(int32 InBufferIndex) const override
+	TSharedPtr<TArray64<uint8>, ESPMode::ThreadSafe> GetBufferDataByIndex(int32 InBufferIndex) const override
 	{
 		if (InBufferIndex == 0)
 		{
@@ -161,7 +161,7 @@ public:
 		return false;
 	}
 
-	EElectraDecoderPlatformPixelFormat GetBufferFormatByIndex(int32 InBufferIndex) const override
+	EPixelFormat GetBufferFormatByIndex(int32 InBufferIndex) const override
 	{
 		if (InBufferIndex == 0)
 		{
@@ -171,9 +171,9 @@ public:
 		{
 			return Buffer[Buffer_Alpha].BufferFormat;
 		}
-	return EElectraDecoderPlatformPixelFormat::INVALID;
+		return EPixelFormat::PF_Unknown;
 	}
-	EElectraDecoderPlatformPixelEncoding GetBufferEncodingByIndex(int32 InBufferIndex) const override
+	EElectraTextureSamplePixelEncoding GetBufferEncodingByIndex(int32 InBufferIndex) const override
 	{
 		if (InBufferIndex == 0)
 		{
@@ -183,7 +183,7 @@ public:
 		{
 			return Buffer[Buffer_Alpha].BufferEncoding;
 		}
-		return EElectraDecoderPlatformPixelEncoding::Native;
+		return EElectraTextureSamplePixelEncoding::Native;
 	}
 	int32 GetBufferPitchByIndex(int32 InBufferIndex) const override
 	{
@@ -216,22 +216,24 @@ public:
 	uint32 Codec4CC = 0;
 	int32 NumBuffers = 0;
 
-	enum {
+	enum 
+	{
 		Buffer_Color = 0,
 		Buffer_Alpha,
-
 		Buffer_Max
 	};
 
-	struct FBufferInfo {
-		TSharedPtr<TArray<uint8>, ESPMode::ThreadSafe> Buffer;
-		EElectraDecoderPlatformPixelFormat BufferFormat;
-		EElectraDecoderPlatformPixelEncoding BufferEncoding;
+	struct FBufferInfo
+	{
+		TSharedPtr<TArray64<uint8>, ESPMode::ThreadSafe> Buffer;
+		EPixelFormat BufferFormat;
+		EElectraTextureSamplePixelEncoding BufferEncoding;
 		int32 Pitch = 0;
 #if ELECTRA_MEDIAGPUBUFFER_DX12
 		FElectraMediaDecoderOutputBufferPool_DX12::FOutputData GPUBuffer;
 #endif
-	} Buffer[Buffer_Max];
+	};
+	FBufferInfo Buffer[Buffer_Max];
 #if ELECTRA_MEDIAGPUBUFFER_DX12
 	TSharedPtr<FElectraDecodersPlatformResources::FAsyncConsecutiveTaskSync, ESPMode::ThreadSafe> Buffer_TaskSync;
 #endif
@@ -523,35 +525,45 @@ bool FVideoDecoderHAPElectra::ResetToCleanStart()
 	return true;
 }
 
-static bool ConvertHapTextureFormat(HapTextureFormat HapFmt, EElectraDecoderPlatformPixelFormat& Format, EElectraDecoderPlatformPixelEncoding& Encoding)
+static bool ConvertHapTextureFormat(HapTextureFormat HapFmt, EPixelFormat& Format, EElectraTextureSamplePixelEncoding& Encoding)
 {
 	switch (HapFmt)
 	{
 		case HapTextureFormat_RGB_DXT1:
-			Format = EElectraDecoderPlatformPixelFormat::DXT1;
-			Encoding = EElectraDecoderPlatformPixelEncoding::Native;
+		{
+			Format = EPixelFormat::PF_DXT1;
+			Encoding = EElectraTextureSamplePixelEncoding::Native;
 			break;
+		}
 		case HapTextureFormat_RGBA_DXT5:
-			Format = EElectraDecoderPlatformPixelFormat::DXT5;
-			Encoding = EElectraDecoderPlatformPixelEncoding::Native;
+		{
+			Format = EPixelFormat::PF_DXT5;
+			Encoding = EElectraTextureSamplePixelEncoding::Native;
 			break;
+		}
 		case HapTextureFormat_YCoCg_DXT5:
-			Format = EElectraDecoderPlatformPixelFormat::DXT5;
-			Encoding = EElectraDecoderPlatformPixelEncoding::YCoCg;
+		{
+			Format = EPixelFormat::PF_DXT5;
+			Encoding = EElectraTextureSamplePixelEncoding::YCoCg;
 			break;
+		}
 		case HapTextureFormat_A_RGTC1:
-			Format = EElectraDecoderPlatformPixelFormat::BC4;
-			Encoding = EElectraDecoderPlatformPixelEncoding::Native;
+		{
+			Format = EPixelFormat::PF_BC4;
+			Encoding = EElectraTextureSamplePixelEncoding::Native;
 			break;
+		}
 		case HapTextureFormat_RGBA_BPTC_UNORM:
 		case HapTextureFormat_RGB_BPTC_UNSIGNED_FLOAT:
 		case HapTextureFormat_RGB_BPTC_SIGNED_FLOAT:
 		default:
-			Format = EElectraDecoderPlatformPixelFormat::INVALID;
-			Encoding = EElectraDecoderPlatformPixelEncoding::Native;
+		{
+			Format = EPixelFormat::PF_Unknown;
+			Encoding = EElectraTextureSamplePixelEncoding::Native;
 			break;
+		}
 	}
-	return Format != EElectraDecoderPlatformPixelFormat::INVALID;
+	return Format != EPixelFormat::PF_Unknown;
 }
 
 static uint32 GetImageBufferPitch(HapTextureFormat HapFmt, uint32 Width)
@@ -700,7 +712,7 @@ IElectraDecoder::EDecoderError FVideoDecoderHAPElectra::DecodeAccessUnit(const F
 			else
 #endif
 			{
-				BufferInfo.Buffer = MakeShared<TArray<uint8>, ESPMode::ThreadSafe>();
+				BufferInfo.Buffer = MakeShared<TArray64<uint8>, ESPMode::ThreadSafe>();
 				BufferInfo.Buffer->AddUninitialized(InstanceBufferData[BufIdx].AllocSize);
 				BufferAddr = BufferInfo.Buffer->GetData();
 				if (BufferAddr == nullptr)
diff --git a/Engine/Plugins/Media/OpusDecoderElectra/Source/OpusDecoderElectra/OpusDecoderElectra.Build.cs b/Engine/Plugins/Media/OpusDecoderElectra/Source/OpusDecoderElectra/OpusDecoderElectra.Build.cs
--- a/Engine/Plugins/Media/OpusDecoderElectra/Source/OpusDecoderElectra/OpusDecoderElectra.Build.cs
+++ b/Engine/Plugins/Media/OpusDecoderElectra/Source/OpusDecoderElectra/OpusDecoderElectra.Build.cs
@@ -9,6 +9,7 @@ namespace UnrealBuildTool.Rules
 			PrivateDependencyModuleNames.AddRange(
 				new string[] {
 					"Core",
+					"ElectraBase",
 					"ElectraCodecFactory",
 					"ElectraDecoders",
 					"libOpus"
diff --git a/Engine/Plugins/Media/OpusDecoderElectra/Source/OpusDecoderElectra/Private/OpusDecoder/ElectraMediaOpusDecoder.cpp b/Engine/Plugins/Media/OpusDecoderElectra/Source/OpusDecoderElectra/Private/OpusDecoder/ElectraMediaOpusDecoder.cpp
--- a/Engine/Plugins/Media/OpusDecoderElectra/Source/OpusDecoderElectra/Private/OpusDecoder/ElectraMediaOpusDecoder.cpp
+++ b/Engine/Plugins/Media/OpusDecoderElectra/Source/OpusDecoderElectra/Private/OpusDecoder/ElectraMediaOpusDecoder.cpp
@@ -14,7 +14,7 @@
 #include "IElectraDecoderOutputAudio.h"
 #include "ElectraDecodersUtils.h"
 #include "Utils/ElectraBitstreamProcessorDefault.h"
-#include "Utils/MPEG/ElectraUtilsMP4.h"		// for parsing the 'dOps' box
+#include "Utilities/UtilitiesMP4.h"		// for parsing the 'dOps' box
 
 #include "opus_multistream.h"
 
@@ -320,7 +320,7 @@ bool FElectraOpusDecoder::Parse_dOps(FDOpsConfig& OutConfig, const TArray<uint8>
 		return !bFailOnError ? false : PostError(0, TEXT("Incomplete 'dOps' box"), ERRCODE_INTERNAL_FAILED_TO_PARSE_CSD);
 	}
 
-	ElectraDecodersUtil::FMP4AtomReader rd(IndOpsBox.GetData(), IndOpsBox.Num());
+	Electra::UtilitiesMP4::FMP4AtomReader rd(IndOpsBox);
 	uint32 Value32;
 	uint16 Value16;
 	uint8 Value8;
diff --git a/Engine/Plugins/Media/VPxDecoderElectra/Source/VPxDecoderElectra/Private/VPxDecoder/ElectraMediaVPxDecoder.cpp b/Engine/Plugins/Media/VPxDecoderElectra/Source/VPxDecoderElectra/Private/VPxDecoder/ElectraMediaVPxDecoder.cpp
--- a/Engine/Plugins/Media/VPxDecoderElectra/Source/VPxDecoderElectra/Private/VPxDecoder/ElectraMediaVPxDecoder.cpp
+++ b/Engine/Plugins/Media/VPxDecoderElectra/Source/VPxDecoderElectra/Private/VPxDecoder/ElectraMediaVPxDecoder.cpp
@@ -23,7 +23,7 @@
 
 #if ENABLE_GPU_BUFFER_HELPER
 #include "ElectraDecodersPlatformResources.h"
-#include COMPILED_PLATFORM_HEADER(ElectraDecoderGPUBufferHelpers.h)
+#include COMPILED_PLATFORM_HEADER(ElectraTextureSampleGPUBufferHelper.h)
 #else
 #define ELECTRA_MEDIAGPUBUFFER_DX12 0
 #endif
@@ -113,7 +113,7 @@ public:
 	{
 		return NumBuffers;
 	}
-	TSharedPtr<TArray<uint8>, ESPMode::ThreadSafe> GetBufferDataByIndex(int32 InBufferIndex) const override
+	TSharedPtr<TArray64<uint8>, ESPMode::ThreadSafe> GetBufferDataByIndex(int32 InBufferIndex) const override
 	{
 		if (InBufferIndex == 0)
 		{
@@ -142,21 +142,21 @@ public:
 #endif
 		return false;
 	}
-	EElectraDecoderPlatformPixelFormat GetBufferFormatByIndex(int32 InBufferIndex) const override
+	EPixelFormat GetBufferFormatByIndex(int32 InBufferIndex) const override
 	{
 		if (InBufferIndex == 0)
 		{
 			return ColorBufferFormat;
 		}
-		return EElectraDecoderPlatformPixelFormat::INVALID;
+		return EPixelFormat::PF_Unknown;
 	}
-	EElectraDecoderPlatformPixelEncoding GetBufferEncodingByIndex(int32 InBufferIndex) const override
+	EElectraTextureSamplePixelEncoding GetBufferEncodingByIndex(int32 InBufferIndex) const override
 	{
 		if (InBufferIndex == 0)
 		{
 			return ColorBufferEncoding;
 		}
-		return EElectraDecoderPlatformPixelEncoding::Native;
+		return EElectraTextureSamplePixelEncoding::Native;
 	}
 	int32 GetBufferPitchByIndex(int32 InBufferIndex) const override
 	{
@@ -185,9 +185,9 @@ public:
 
 	uint32 Codec4CC = 0;
 	int32 NumBuffers = 0;
-	TSharedPtr<TArray<uint8>, ESPMode::ThreadSafe> ColorBuffer;
-	EElectraDecoderPlatformPixelFormat ColorBufferFormat;
-	EElectraDecoderPlatformPixelEncoding ColorBufferEncoding;
+	TSharedPtr<TArray64<uint8>, ESPMode::ThreadSafe> ColorBuffer;
+	EPixelFormat ColorBufferFormat;
+	EElectraTextureSamplePixelEncoding ColorBufferEncoding;
 	int32 ColorPitch = 0;
 #if ELECTRA_MEDIAGPUBUFFER_DX12
 	FElectraMediaDecoderOutputBufferPool_DX12::FOutputData GPUBuffer;
@@ -934,8 +934,8 @@ FVideoDecoderVPxElectra::EConvertResult FVideoDecoderVPxElectra::ConvertDecoderO
 		if (InDecodedAlphaImage == nullptr)
 		{
 			NewOutput->NumBuffers = 1;
-			NewOutput->ColorBufferFormat = EElectraDecoderPlatformPixelFormat::NV12;
-			NewOutput->ColorBufferEncoding = EElectraDecoderPlatformPixelEncoding::Native;
+			NewOutput->ColorBufferFormat = EPixelFormat::PF_NV12;
+			NewOutput->ColorBufferEncoding = EElectraTextureSamplePixelEncoding::Native;
 			NewOutput->DecodedWidth = NewOutput->Width;
 			NewOutput->DecodedHeight = bUseGPUBuffers ? NewOutput->Height : (NewOutput->Height * 3 / 2);
 			if (!ConvertDecodedImageToNV12orP010(NewOutput.Get(), InDecodedImage, PlatformDevice, PlatformDeviceVersion))
@@ -947,8 +947,8 @@ FVideoDecoderVPxElectra::EConvertResult FVideoDecoderVPxElectra::ConvertDecoderO
 		else
 		{
 			NewOutput->NumBuffers = 1;
-			NewOutput->ColorBufferFormat = EElectraDecoderPlatformPixelFormat::B8G8R8A8;
-			NewOutput->ColorBufferEncoding = EElectraDecoderPlatformPixelEncoding::YCbCr_Alpha;
+			NewOutput->ColorBufferFormat = EPixelFormat::PF_B8G8R8A8;
+			NewOutput->ColorBufferEncoding = EElectraTextureSamplePixelEncoding::YCbCr_Alpha;
 			NewOutput->DecodedWidth = NewOutput->Width;
 			NewOutput->DecodedHeight = NewOutput->Height;
 			if (!ConvertDecodedImageWithAlpha(NewOutput.Get(), InDecodedImage, InDecodedAlphaImage, PlatformDevice, PlatformDeviceVersion))
@@ -967,8 +967,8 @@ FVideoDecoderVPxElectra::EConvertResult FVideoDecoderVPxElectra::ConvertDecoderO
 			return EConvertResult::Failure;
 		}
 		NewOutput->NumBuffers = 1;
-		NewOutput->ColorBufferFormat = EElectraDecoderPlatformPixelFormat::P010;
-		NewOutput->ColorBufferEncoding = EElectraDecoderPlatformPixelEncoding::Native;
+		NewOutput->ColorBufferFormat = EPixelFormat::PF_P010;
+		NewOutput->ColorBufferEncoding = EElectraTextureSamplePixelEncoding::Native;
 		NewOutput->DecodedWidth = NewOutput->Width;
 		NewOutput->DecodedHeight = bUseGPUBuffers ? NewOutput->Height : (NewOutput->Height * 3 / 2);
 
@@ -1024,7 +1024,7 @@ bool FVideoDecoderVPxElectra::ConvertDecodedImageToNV12orP010(FVideoDecoderOutpu
 	uint8* DstUV;
 	uint32 DstPitch;
 
-	TSharedPtr<TArray<uint8>, ESPMode::ThreadSafe>& OutNV12Buffer = InNewOutput->ColorBuffer;
+	TSharedPtr<TArray64<uint8>, ESPMode::ThreadSafe>& OutNV12Buffer = InNewOutput->ColorBuffer;
 #if ELECTRA_MEDIAGPUBUFFER_DX12
 	if (!PlatformDevice || PlatformDeviceVersion < 12000)
 #endif
@@ -1037,7 +1037,7 @@ bool FVideoDecoderVPxElectra::ConvertDecodedImageToNV12orP010(FVideoDecoderOutpu
 		InNewOutput->GPUBuffer.Fence = nullptr;
 #endif
 
-		OutNV12Buffer = MakeShared<TArray<uint8>, ESPMode::ThreadSafe>();
+		OutNV12Buffer = MakeShared<TArray64<uint8>, ESPMode::ThreadSafe>();
 		OutNV12Buffer->AddUninitialized(AllocSize);
 		DstY = OutNV12Buffer->GetData();
 		DstUV = DstY + Pitch * ah;
@@ -1162,7 +1162,7 @@ bool FVideoDecoderVPxElectra::ConvertDecodedImageWithAlpha(FVideoDecoderOutputVP
 	}
 
 
-	TSharedPtr<TArray<uint8>, ESPMode::ThreadSafe>& OutPixelBuffer = InNewOutput->ColorBuffer;
+	TSharedPtr<TArray64<uint8>, ESPMode::ThreadSafe>& OutPixelBuffer = InNewOutput->ColorBuffer;
 	uint16* Dst0 = nullptr;
 	uint16* Dst1 = nullptr;
 	int32 DstPitch = 0;
@@ -1178,7 +1178,7 @@ bool FVideoDecoderVPxElectra::ConvertDecodedImageWithAlpha(FVideoDecoderOutputVP
 #endif
 
 		int32 AllocSize = aw * ah * sizeof(uint16) * 4;
-		OutPixelBuffer = MakeShared<TArray<uint8>, ESPMode::ThreadSafe>();
+		OutPixelBuffer = MakeShared<TArray64<uint8>, ESPMode::ThreadSafe>();
 		OutPixelBuffer->AddUninitialized(AllocSize);
 		Dst0 = (uint16*) OutPixelBuffer->GetData();
 		InNewOutput->ColorPitch = aw * sizeof(uint16) * 4;
@@ -1214,8 +1214,8 @@ bool FVideoDecoderVPxElectra::ConvertDecodedImageWithAlpha(FVideoDecoderOutputVP
 	}
 #endif
 	InNewOutput->NumBits = 16;
-	InNewOutput->ColorBufferFormat = EElectraDecoderPlatformPixelFormat::A16B16G16R16;
-	InNewOutput->ColorBufferEncoding = EElectraDecoderPlatformPixelEncoding::YCbCr_Alpha;
+	InNewOutput->ColorBufferFormat = EPixelFormat::PF_A16B16G16R16;
+	InNewOutput->ColorBufferEncoding = EElectraTextureSamplePixelEncoding::YCbCr_Alpha;
 
 	const int32 LineSkipY = PitchY - w;
 	const int32 LineSkipUV = PitchU - w/2;
