ElectraSamples: Added missing export macro

#rnx

diff --git a/Engine/Plugins/Media/ElectraUtil/Source/ElectraSamples/Public/Linux/ElectraTextureSample.h b/Engine/Plugins/Media/ElectraUtil/Source/ElectraSamples/Public/Linux/ElectraTextureSample.h
--- a/Engine/Plugins/Media/ElectraUtil/Source/ElectraSamples/Public/Linux/ElectraTextureSample.h
+++ b/Engine/Plugins/Media/ElectraUtil/Source/ElectraSamples/Public/Linux/ElectraTextureSample.h
@@ -14,18 +14,18 @@ class FElectraTextureSample
 {
 public:
 	FElectraTextureSample() = default;
-	virtual	~FElectraTextureSample();
+	ELECTRASAMPLES_API ~FElectraTextureSample();
 
-	bool FinishInitialization() override;
+	ELECTRASAMPLES_API bool FinishInitialization() override;
 
-	const void* GetBuffer() override;
-	uint32 GetStride() const override;
-	EMediaTextureSampleFormat GetFormat() const override;
-	FRHITexture* GetTexture() const override;
-    IMediaTextureSampleConverter* GetMediaTextureSampleConverter() override;
+	ELECTRASAMPLES_API const void* GetBuffer() override;
+	ELECTRASAMPLES_API uint32 GetStride() const override;
+	ELECTRASAMPLES_API EMediaTextureSampleFormat GetFormat() const override;
+	ELECTRASAMPLES_API FRHITexture* GetTexture() const override;
+    ELECTRASAMPLES_API IMediaTextureSampleConverter* GetMediaTextureSampleConverter() override;
 
 #if !UE_SERVER
-	void ShutdownPoolable() override;
+	ELECTRASAMPLES_API void ShutdownPoolable() override;
 #endif
 
 	TSharedPtr<TArray64<uint8>, ESPMode::ThreadSafe> Buffer;
