ElectraSamples: Added empty methods for server builds

#rnx

diff --git a/Engine/Plugins/Media/ElectraUtil/Source/ElectraSamples/Private/ElectraTextureSampleBase.cpp b/Engine/Plugins/Media/ElectraUtil/Source/ElectraSamples/Private/ElectraTextureSampleBase.cpp
--- a/Engine/Plugins/Media/ElectraUtil/Source/ElectraSamples/Private/ElectraTextureSampleBase.cpp
+++ b/Engine/Plugins/Media/ElectraUtil/Source/ElectraSamples/Private/ElectraTextureSampleBase.cpp
@@ -2,14 +2,14 @@
 
 #include "Misc/Build.h"
 
-#if !UE_SERVER
-
 #include "HAL/IConsoleManager.h"
 #include "HAL/Platform.h"
 #include "ElectraTextureSample.h"
 #include "ElectraTextureSampleUtils.h"
 #include "ElectraSamplesModule.h"
 
+#if !UE_SERVER
+
 #include "MediaDecoderOutput.h"
 #include "MediaVideoDecoderOutput.h"
 
@@ -449,4 +449,97 @@ float IElectraTextureSampleBase::GetSampleDataScale(bool b10Bit) const
 	return PixelDataScale;
 }
 
+
+#else
+
+void IElectraTextureSampleBase::InitializeCommon(const Electra::FParamDict& InParams)
+{ }
+
+bool IElectraTextureSampleBase::FinishInitialization()
+{ return true; }
+
+IElectraTextureSampleBase::FReleaseDelegate& IElectraTextureSampleBase::GetReleaseDelegate()
+{ return ReleaseDelegate; }
+
+bool IElectraTextureSampleBase::IsCacheable() const
+{ return true; }
+
+FIntPoint IElectraTextureSampleBase::GetDim() const
+{ return FIntPoint(); }
+
+void IElectraTextureSampleBase::SetDim(const FIntPoint& InDim)
+{ }
+
+FIntPoint IElectraTextureSampleBase::GetOutputDim() const
+{ return FIntPoint(); }
+
+FMediaTimeStamp IElectraTextureSampleBase::GetTime() const
+{ return FMediaTimeStamp(); }
+
+void IElectraTextureSampleBase::SetTime(const FMediaTimeStamp& InTime)
+{ }
+
+void IElectraTextureSampleBase::SetDuration(const FTimespan& InDuration)
+{ }
+
+TSharedPtr<const IVideoDecoderColorimetry, ESPMode::ThreadSafe> IElectraTextureSampleBase::GetColorimetry() const
+{ return nullptr; }
+
+EPixelFormat IElectraTextureSampleBase::GetPixelFormat() const
+{ return EPixelFormat::PF_Unknown; }
+
+EElectraTextureSamplePixelEncoding IElectraTextureSampleBase::GetPixelFormatEncoding() const
+{ return EElectraTextureSamplePixelEncoding::Native; }
+
+FTimespan IElectraTextureSampleBase::GetDuration() const
+{ return FTimespan(); }
+
+TOptional<FTimecode> IElectraTextureSampleBase::GetTimecode() const
+{ return TOptional<FTimecode>(); }
+
+TOptional<FFrameRate> IElectraTextureSampleBase::GetFramerate() const
+{ return TOptional<FFrameRate>(); }
+
+double IElectraTextureSampleBase::GetAspectRatio() const
+{ return 1.0; }
+
+EMediaOrientation IElectraTextureSampleBase::GetOrientation() const
+{ return EMediaOrientation::Original; }
+
+bool IElectraTextureSampleBase::IsOutputSrgb() const
+{ return true; }
+
+const FMatrix& IElectraTextureSampleBase::GetYUVToRGBMatrix() const
+{ return FMatrix::Identity; }
+
+bool IElectraTextureSampleBase::GetFullRange() const
+{ return false; }
+
+FMatrix44f IElectraTextureSampleBase::GetSampleToRGBMatrix() const
+{ return SampleToRgbMtx; }
+
+const UE::Color::FColorSpace& IElectraTextureSampleBase::GetSourceColorSpace() const
+{ return SourceColorSpace; }
+
+UE::Color::EEncoding IElectraTextureSampleBase::GetEncodingType() const
+{ return ColorEncoding; }
+
+float IElectraTextureSampleBase::GetHDRNitsNormalizationFactor() const
+{ return 1.0f; }
+
+bool IElectraTextureSampleBase::GetDisplayMasteringLuminance(float& OutMin, float& OutMax) const
+{ return false; }
+
+TOptional<UE::Color::FColorSpace> IElectraTextureSampleBase::GetDisplayMasteringColorSpace() const
+{ return DisplayMasteringColorSpace; }
+
+bool IElectraTextureSampleBase::GetMaxLuminanceLevels(uint16& OutCLL, uint16& OutFALL) const
+{ return false; }
+
+MediaShaders::EToneMapMethod IElectraTextureSampleBase::GetToneMapMethod() const
+{ return MediaShaders::EToneMapMethod::None; }
+
+float IElectraTextureSampleBase::GetSampleDataScale(bool b10Bit) const
+{ return 1.0f; }
+
 #endif
