ElectraSamples: proper error handling instead of check() macros in DX12 buffer helper

#rnx

diff --git a/Engine/Plugins/Media/ElectraUtil/Source/ElectraSamples/Public/Windows/WindowsElectraTextureSampleGPUBufferHelper.h b/Engine/Plugins/Media/ElectraUtil/Source/ElectraSamples/Public/Windows/WindowsElectraTextureSampleGPUBufferHelper.h
--- a/Engine/Plugins/Media/ElectraUtil/Source/ElectraSamples/Public/Windows/WindowsElectraTextureSampleGPUBufferHelper.h
+++ b/Engine/Plugins/Media/ElectraUtil/Source/ElectraSamples/Public/Windows/WindowsElectraTextureSampleGPUBufferHelper.h
@@ -224,7 +224,6 @@ private:
 		HeapDesc.Properties = CD3DX12_HEAP_PROPERTIES(D3D12HeapType);
 		HeapDesc.Flags = D3D12_HEAP_FLAG_CREATE_NOT_ZEROED | ((D3D12HeapType == D3D12_HEAP_TYPE_UPLOAD) ? D3D12_HEAP_FLAG_ALLOW_ONLY_BUFFERS : D3D12_HEAP_FLAG_ALLOW_ONLY_NON_RT_DS_TEXTURES);
 		HRESULT Res = D3D12Device->CreateHeap(&HeapDesc, IID_PPV_ARGS(D3D12OutputHeap.GetInitReference()));
-		check(SUCCEEDED(Res));
 		if (Res != S_OK)
 		{
 			FreeMask = 0;
@@ -261,8 +260,7 @@ private:
 			// Make a placed resource to use...
 			D3D12_RESOURCE_DESC Desc = {};
 			InitializeDesc(Desc);
-			HRESULT Res = D3D12Device->CreatePlacedResource(D3D12OutputHeap, BufferSize * BufferIdx, &Desc, (D3D12HeapType == D3D12_HEAP_TYPE_UPLOAD) ? D3D12_RESOURCE_STATE_GENERIC_READ : D3D12_RESOURCE_STATE_COMMON, nullptr, IID_PPV_ARGS(Resource.GetInitReference()));
-			check(SUCCEEDED(Res));
+			HRESULT Res = D3D12Device->CreatePlacedResource(D3D12OutputHeap, BufferSize, &Desc, (D3D12HeapType == D3D12_HEAP_TYPE_UPLOAD) ? D3D12_RESOURCE_STATE_GENERIC_READ : D3D12_RESOURCE_STATE_COMMON, nullptr, IID_PPV_ARGS(Resource.GetInitReference()));
 			if (Res == S_OK)
 			{
 #if !UE_BUILD_SHIPPING
@@ -274,12 +272,22 @@ private:
 				// Register a destruction callback so we can reset our memory management here transparently...
 				TRefCountPtr<ID3DDestructionNotifier> Notifier;
 				Res = Resource->QueryInterface(__uuidof(ID3DDestructionNotifier), (void**)Notifier.GetInitReference());
-				check(SUCCEEDED(Res));
-
-				// note: we keep a reference to the heap in our context data for the destruction callback, so we can ensure the heap lives as long as there are placed resources in it
-				UINT CallbackID;
-				Res = Notifier->RegisterDestructionCallback(ResourceDestructionCallback, new FResourceInfo(InPoolWeakRef, BlockIdx, BufferIdx, D3D12OutputHeap), &CallbackID);
-				check(SUCCEEDED(Res));
+				if (SUCCEEDED(Res))
+				{
+					// note: we keep a reference to the heap in our context data for the destruction callback, so we can ensure the heap lives as long as there are placed resources in it
+					UINT CallbackID;
+					Res = Notifier->RegisterDestructionCallback(ResourceDestructionCallback, new FResourceInfo(InPoolWeakRef, BlockIdx, BufferIdx, D3D12OutputHeap), &CallbackID);
+					if (!(SUCCEEDED(Res)))
+					{
+						Resource = nullptr;
+						FreeBuffer(BufferIdx);
+					}
+				}
+				else
+				{
+					Resource = nullptr;
+					FreeBuffer(BufferIdx);
+				}
 			}
 			else
 			{
@@ -512,10 +520,12 @@ private:
 	void InitCommon()
 	{
 		HRESULT Res = D3D12Device->CreateFence(0, D3D12_FENCE_FLAG_NONE, IID_PPV_ARGS(D3D12BufferFence.GetInitReference()));
-		check(SUCCEEDED(Res));
+		if (SUCCEEDED(Res))
+		{
 #if !UE_BUILD_SHIPPING
-		D3D12BufferFence->SetName(TEXT("ElectraOutputBufferPoolFence"));
+			D3D12BufferFence->SetName(TEXT("ElectraOutputBufferPoolFence"));
 #endif
+		}
 	}
 
 	TArray<FElectraMediaDecoderOutputBufferPoolBlock_DX12*> Blocks;
