ElectraBase: Preventing an out of bounds access when parsing an empty file:// URL by checking the remaining length first.

#rnx

diff --git a/Engine/Plugins/Media/ElectraUtil/Source/ElectraBase/Private/URLParser.cpp b/Engine/Plugins/Media/ElectraUtil/Source/ElectraBase/Private/URLParser.cpp
--- a/Engine/Plugins/Media/ElectraUtil/Source/ElectraBase/Private/URLParser.cpp
+++ b/Engine/Plugins/Media/ElectraUtil/Source/ElectraBase/Private/URLParser.cpp
@@ -144,7 +144,7 @@ namespace Electra
 						++it;
 
 						// Check for a third '/' in Windows filenames like "file:///c:/autoexec.bat"
-						if (bIsFile && IsPathSeparator(*it) && it.GetRemainingLength() > 3)
+						if (bIsFile && it.GetRemainingLength() > 3 && IsPathSeparator(*it))
 						{
 							const TCHAR* Next = it.GetRemainder();
 							if (Next[2] == TCHAR(':') && Next[3] == TCHAR('/'))
