Added back accidentally removed resource offset

#rnx

diff --git a/Engine/Plugins/Media/ElectraUtil/Source/ElectraSamples/Public/Windows/WindowsElectraTextureSampleGPUBufferHelper.h b/Engine/Plugins/Media/ElectraUtil/Source/ElectraSamples/Public/Windows/WindowsElectraTextureSampleGPUBufferHelper.h
--- a/Engine/Plugins/Media/ElectraUtil/Source/ElectraSamples/Public/Windows/WindowsElectraTextureSampleGPUBufferHelper.h
+++ b/Engine/Plugins/Media/ElectraUtil/Source/ElectraSamples/Public/Windows/WindowsElectraTextureSampleGPUBufferHelper.h
@@ -260,7 +260,7 @@ private:
 			// Make a placed resource to use...
 			D3D12_RESOURCE_DESC Desc = {};
 			InitializeDesc(Desc);
-			HRESULT Res = D3D12Device->CreatePlacedResource(D3D12OutputHeap, BufferSize, &Desc, (D3D12HeapType == D3D12_HEAP_TYPE_UPLOAD) ? D3D12_RESOURCE_STATE_GENERIC_READ : D3D12_RESOURCE_STATE_COMMON, nullptr, IID_PPV_ARGS(Resource.GetInitReference()));
+			HRESULT Res = D3D12Device->CreatePlacedResource(D3D12OutputHeap, BufferSize * BufferIdx, &Desc, (D3D12HeapType == D3D12_HEAP_TYPE_UPLOAD) ? D3D12_RESOURCE_STATE_GENERIC_READ : D3D12_RESOURCE_STATE_COMMON, nullptr, IID_PPV_ARGS(Resource.GetInitReference()));
 			if (Res == S_OK)
 			{
 #if !UE_BUILD_SHIPPING
