Electra: Minor source code cleanup; removal of obsolete and unused class

#rnx

diff --git a/Engine/Plugins/Media/ElectraPlayer/Source/ElectraPlayerRuntime/Private/ElectraPlayer.h b/Engine/Plugins/Media/ElectraPlayer/Source/ElectraPlayerRuntime/Private/ElectraPlayer.h
--- a/Engine/Plugins/Media/ElectraPlayer/Source/ElectraPlayerRuntime/Private/ElectraPlayer.h
+++ b/Engine/Plugins/Media/ElectraPlayer/Source/ElectraPlayerRuntime/Private/ElectraPlayer.h
@@ -44,8 +44,6 @@ using IMetaDataDecoderOutputPtr = TSharedPtr<IMetaDataDecoderOutput, ESPMode::Th
 class ISubtitleDecoderOutput;
 using ISubtitleDecoderOutputPtr = TSharedPtr<ISubtitleDecoderOutput, ESPMode::ThreadSafe>;
 
-class FElectraRendererAudio;
-
 using namespace Electra;
 
 
diff --git a/Engine/Plugins/Media/ElectraUtil/Source/ElectraBase/Public/MediaAudioDecoderOutput.h b/Engine/Plugins/Media/ElectraUtil/Source/ElectraBase/Public/MediaAudioDecoderOutput.h
--- a/Engine/Plugins/Media/ElectraUtil/Source/ElectraBase/Public/MediaAudioDecoderOutput.h
+++ b/Engine/Plugins/Media/ElectraUtil/Source/ElectraBase/Public/MediaAudioDecoderOutput.h
@@ -2,40 +2,3 @@
 #pragma once
 
 #include "MediaDecoderOutput.h"
-
-class IAudioDecoderOutput : public IDecoderOutput
-{
-public:
-	enum class ESampleFormat
-	{
-		Undefined = 0,
-		Double,
-		Float,
-		Int8,
-		Int16,
-		Int32
-	};
-
-	virtual ~IAudioDecoderOutput() {}
-
-	virtual void Reserve(uint32 InBufferSizeBytes) = 0;
-
-	virtual void Initialize(ESampleFormat InFormat, uint32 InNumChannels, uint32 InSampleRate, FTimespan InDuration, const FDecoderTimeStamp& InPts, uint32 InBufferSizeBytes) = 0;
-
-	virtual const void* GetBuffer() const = 0;
-	virtual uint32 GetUsedBufferBytes() const = 0;
-	virtual uint32 GetReservedBufferBytes() const = 0;
-
-	virtual uint32 GetChannels() const = 0;
-	virtual uint32 GetFrames() const = 0;
-	virtual ESampleFormat GetFormat() const = 0;
-
-	virtual FDecoderTimeStamp GetTime() const = 0;
-	virtual FTimespan GetDuration() const = 0;
-
-	virtual uint32 GetSampleRate() const = 0;
-
-	virtual void SetTime(const FDecoderTimeStamp& InTime) = 0;
-};
-
-using IAudioDecoderOutputPtr = TSharedPtr<IAudioDecoderOutput, ESPMode::ThreadSafe>;
diff --git a/Engine/Plugins/Media/ElectraUtil/Source/ElectraBase/Public/MediaDecoderOutput.h b/Engine/Plugins/Media/ElectraUtil/Source/ElectraBase/Public/MediaDecoderOutput.h
--- a/Engine/Plugins/Media/ElectraUtil/Source/ElectraBase/Public/MediaDecoderOutput.h
+++ b/Engine/Plugins/Media/ElectraUtil/Source/ElectraBase/Public/MediaDecoderOutput.h
@@ -2,15 +2,6 @@
 #pragma once
 
 #include "Misc/Timespan.h"
-#include "Math/IntPoint.h"
-#include "Templates/Atomic.h"
-#include "Templates/SharedPointer.h"
-#include "Templates/RefCounting.h"
-#include "Containers/Queue.h"
-#include "Misc/ScopeLock.h"
-#include "HAL/PlatformProcess.h"
-#include "HAL/PlatformMisc.h"
-#include "Tickable.h"
 
 #include "ParameterDictionary.h"
 
@@ -24,305 +15,6 @@ struct FDecoderTimeStamp
 };
 
 
-class IDecoderOutputPoolable : public TSharedFromThis<IDecoderOutputPoolable, ESPMode::ThreadSafe>
-{
-public:
-	virtual void InitializePoolable() { }
-	virtual void ShutdownPoolable() { }
-	virtual bool IsReadyForReuse() { return true; }
-
-public:
-	virtual ~IDecoderOutputPoolable() { }
-};
-
-
-template<typename ObjectType> class TElectraPoolDefaultObjectFactory
-{
-public:
-	static ObjectType* Create()
-	{
-		return new ObjectType;
-	}
-};
-
-
-/**
- * This class defines a pool of objects that inherit from the `IDecoderOutputPoolable` interface.
- * The pool has no bound on the number of objects created, it merely tracks which objects have
- * been handed out and awaits their return when the pool is destroyed.
- */
-template<typename ObjectType, typename ObjectFactory = TElectraPoolDefaultObjectFactory<ObjectType>>
-class TDecoderOutputObjectPool
-{
-public:
-	/**
-	 * Creates a new object pool.
-	 * The pool has a custom deleter which only marks the pool as expired, but keeps it
-	 * around until all the elements it has ever handed out are returned and indicate
-	 * that they are no longer being used.
-	 */
-	static TSharedPtr<TDecoderOutputObjectPool, ESPMode::ThreadSafe> Create()
-	{
-		return MakeShareable(new TDecoderOutputObjectPool(), FPoolDeleter());
-	}
-
-	/**
-	 * Acquires an object from the pool.
-	 * This returns either a new object or one that was used before and has been
-	 * returned to the pool. The object will not be re-initialized by default.
-	 * For this your managed object needs to implement the `InitializePoolable()`
-	 * method. Likewise, for object members that should be freed before the object
-	 * enters the pool for re-use, the object needs to implement the `ShutdownPoolable()`
-	 * method.
-	 */
-	TSharedRef<ObjectType, ESPMode::ThreadSafe> AcquireShared()
-	{
-		return MakeShareable(ObjectPool->Acquire(), TObjectDeleter(ObjectPool));
-	}
-
-private:
-	static_assert(TPointerIsConvertibleFromTo<ObjectType, IDecoderOutputPoolable>::Value, "Poolable objects must implement the IDecoderOutputPoolable interface.");
-
-	/**
-	 * This class defines the actual pool.
-	 */
-	template<typename T=ObjectFactory>
-	class TObjectPool : public FTickableGameObject
-	{
-	public:
-		TObjectPool() = default;
-		~TObjectPool() = default;
-
-		/** Called by the enclosing pool when it goes out of scope to let us handle cleanup of in-flight objects. */
-		void SetPendingDestruction(TSharedPtr<TObjectPool<T>, ESPMode::ThreadSafe> InSelf)
-		{
-			Self = MoveTemp(InSelf);
-			bIsPendingDestruction = true;
-		}
-
-		/** Acquire an object from the pool. */
-		ObjectType* Acquire()
-		{
-			// `Acquire()` cannot possibly be called when pool destruction is pending.
-			check(!bIsPendingDestruction);
-
-			ObjectType* Object = nullptr;
-
-			CriticalSection.Lock();
-			// Handle objects that have just become reusable again.
-			HandleNewReturns();
-			// Take an object from the available list, if there is one.
-			if (Available.Num() > 0)
-			{
-				Object = Available.Pop(EAllowShrinking::No);
-			}
-			CriticalSection.Unlock();
-			// If there is no object, we have to create a new one.
-			if (Object == nullptr)
-			{
-				Object = T::Create();
-			}
-			Object->InitializePoolable();
-
-			// Take note of the object we are handing out.
-			CriticalSection.Lock();
-			InFlight.Emplace(Object);
-			CriticalSection.Unlock();
-
-			return Object;
-		}
-
-		/** Return the given object to the pool. */
-		void Release(ObjectType* InObject)
-		{
-			if (InObject == nullptr)
-			{
-				return;
-			}
-
-			FScopeLock Lock(&CriticalSection);
-			check(InFlight.Contains(InObject));
-
-			InFlight.RemoveSwap(InObject, EAllowShrinking::No);
-			if (InObject->IsReadyForReuse())
-			{
-				InObject->ShutdownPoolable();
-				Available.Push(InObject);
-			}
-			else
-			{
-				Returned.Emplace(InObject);
-			}
-		}
-
-		bool IsTickableWhenPaused() const override
-		{ return true; }
-		bool IsTickableInEditor() const override
-		{ return true; }
-		void Tick(float /*InDeltaTime*/) override
-		{
-			if (!bIsPendingDestruction)
-			{
-				return;
-			}
-
-			CriticalSection.Lock();
-			// Handle objects that have just become reusable again.
-			HandleNewReturns();
-			CriticalSection.Unlock();
-			// Release all objects that are available now.
-			// We do not need to lock the mutex for this since `Acquire()` - which would need that -
-			// can no longer get called during a pending destruction as there is no user-code owner any more.
-			while(!Available.IsEmpty())
-			{
-				delete Available.Pop();
-			}
-			// If there are no more in-flight objects or objects that await reusability we can destroy ourselved.
-			if (InFlight.IsEmpty() && Returned.IsEmpty())
-			{
-				// Bye-bye.
-				Self.Reset();
-			}
-		}
-		bool IsTickable() const override
-		{ return true; }
-		TStatId GetStatId() const override
-		{ RETURN_QUICK_DECLARE_CYCLE_STAT(TObjectPool, STATGROUP_Tickables); }
-
-	private:
-		void HandleNewReturns()
-		{
-			// NOTE: This method expects the `CriticalSection` lock to have been locked already
-			//       in order to reduce the number of lock/unlocks.
-
-			// Move all the returned objects that have become ready for reuse into a temporary list.
-			TArray<ObjectType*> NewlyReusableObjects;
-			for(int32 i=Returned.Num()-1; i>=0; --i)
-			{
-				if (Returned[i]->IsReadyForReuse())
-				{
-					NewlyReusableObjects.Emplace(Returned[i]);
-					Returned.RemoveAtSwap(i);
-				}
-			}
-			// If there are objects that have just become reusable we call their `ShutdownPoolable()`
-			// outside of our mutex lock in case what they are doing is costly.
-			if (NewlyReusableObjects.Num())
-			{
-				CriticalSection.Unlock();
-				for(int32 i=NewlyReusableObjects.Num()-1; i>=0; --i)
-				{
-					NewlyReusableObjects[i]->ShutdownPoolable();
-				}
-				CriticalSection.Lock();
-				Available.Append(MoveTemp(NewlyReusableObjects));
-			}
-		}
-
-		/** Critical section for synchronizing access to the free list. */
-		FCriticalSection CriticalSection;
-
-		/** List of available objects. */
-		TArray<ObjectType*> Available;
-
-		/** List of in-flight objects. */
-		TArray<ObjectType*> InFlight;
-
-		/** List of returned objects, waiting for reuseability. */
-		TArray<ObjectType*> Returned;
-
-		/** Pointer to self, which gets set when pool destruction is pending to lock it until complete. */
-		TSharedPtr<TObjectPool<T>, ESPMode::ThreadSafe> Self;
-
-		/** Flag indicating whether or not the pool is pending destruction. */
-		TAtomic<bool> bIsPendingDestruction { false };
-	};
-
-
-	/** Deleter for pooled objects. */
-	class TObjectDeleter
-	{
-	public:
-
-		/** Create and initialize a new instance. */
-		TObjectDeleter(const TSharedPtr<TObjectPool<ObjectFactory>, ESPMode::ThreadSafe>& InOwningPool)
-			: OwningPool(InOwningPool)
-		{ }
-
-		/** Function operator to execute deleter. */
-		void operator()(ObjectType* ObjectToDelete)
-		{
-			TSharedPtr<TObjectPool<ObjectFactory>, ESPMode::ThreadSafe> PinnedStorage = OwningPool.Pin();
-			// The weak pointer must be lockable since the pool has a circular reference to itself
-			// during pool destruction.
-			if (ensure(PinnedStorage.IsValid()))
-			{
-				PinnedStorage->Release(ObjectToDelete);
-			}
-			else
-			{
-				if (ensure(ObjectToDelete->IsReadyForReuse()))
-				{
-					delete ObjectToDelete;
-				}
-			}
-		}
-
-	private:
-		/** Weak pointer to the owning object pool. */
-		TWeakPtr<TObjectPool<ObjectFactory>, ESPMode::ThreadSafe> OwningPool;
-	};
-
-
-	/**
-	 * This deleter class is invoked when the pool is to be destroyed.
-	 * We flag the actual pool as such and let its handling take care of this.
-	 */
-	class FPoolDeleter
-	{
-	public:
-		void operator()(TDecoderOutputObjectPool* InPool)
-		{
-			// Only mark the actual pool for destruction.
-			// It handles itself during `Tick()` and will destroy itself
-			// once all objects are returned it it.
-			check(InPool->ObjectPool.IsValid());
-			InPool->ObjectPool->SetPendingDestruction(MoveTemp(InPool->ObjectPool));
-		}
-	};
-
-	/** Private constructor. */
-	TDecoderOutputObjectPool() : ObjectPool(MakeShareable(new TObjectPool<ObjectFactory>()))
-	{ }
-
-	/** The actual pool object. */
-	TSharedPtr<TObjectPool<ObjectFactory>, ESPMode::ThreadSafe> ObjectPool;
-};
-
-
-class IDecoderOutput;
-
-class IDecoderOutputOwner
-{
-public:
-	virtual void SampleReleasedToPool(IDecoderOutput* InDecoderOutput) = 0;
-};
-
-
-class IDecoderOutput : public IDecoderOutputPoolable
-{
-public:
-	virtual void SetOwner(const TSharedPtr<IDecoderOutputOwner, ESPMode::ThreadSafe>& Renderer) = 0;
-	virtual FDecoderTimeStamp GetTime() const = 0;
-	virtual FTimespan GetDuration() const = 0;
-
-	virtual Electra::FParamDict& GetMutablePropertyDictionary()
-	{ return PropertyDictionary; }
-private:
-	Electra::FParamDict PropertyDictionary;
-};
-
-
 namespace IDecoderOutputOptionNames
 {
 static const FName PTS(TEXT("pts"));
