Fixup API macro usage

#rnx

diff --git a/Engine/Plugins/Media/ElectraUtil/Source/ElectraSamples/Public/Android/ElectraTextureSample.h b/Engine/Plugins/Media/ElectraUtil/Source/ElectraSamples/Public/Android/ElectraTextureSample.h
--- a/Engine/Plugins/Media/ElectraUtil/Source/ElectraSamples/Public/Android/ElectraTextureSample.h
+++ b/Engine/Plugins/Media/ElectraUtil/Source/ElectraSamples/Public/Android/ElectraTextureSample.h
@@ -23,48 +23,48 @@ struct AHardwareBuffer;
 /**
  * Texture sample generated by AndroidMedia player.
  */
-class UE_API FElectraTextureSample final
+class FElectraTextureSample final
 	: public IElectraTextureSampleBase
 	, public IMediaTextureSampleConverter
 	, public TSharedFromThis<FElectraTextureSample, ESPMode::ThreadSafe>
 {
 public:
-	FElectraTextureSample(TSharedPtr<FElectraTextureSampleSupport, ESPMode::ThreadSafe> InSupport);
-	virtual ~FElectraTextureSample();
+	UE_API FElectraTextureSample(TSharedPtr<FElectraTextureSampleSupport, ESPMode::ThreadSafe> InSupport);
+	UE_API virtual ~FElectraTextureSample();
 
-	bool Initialize();
+	UE_API bool Initialize();
 
-	bool FinishInitialization() override;
+	UE_API bool FinishInitialization() override;
 
 	//~ IMediaTextureSample interface
-	const void* GetBuffer() override;
-	uint32 GetStride() const override;
-	EMediaTextureSampleFormat GetFormat() const override;
-	FRHITexture* GetTexture() const override;
-	IMediaTextureSampleConverter* GetMediaTextureSampleConverter() override;
-	FIntPoint GetDim() const override;
+	UE_API const void* GetBuffer() override;
+	UE_API uint32 GetStride() const override;
+	UE_API EMediaTextureSampleFormat GetFormat() const override;
+	UE_API FRHITexture* GetTexture() const override;
+	UE_API IMediaTextureSampleConverter* GetMediaTextureSampleConverter() override;
+	UE_API FIntPoint GetDim() const override;
 
-	bool IsReadyForReuse() override;
-	void ShutdownPoolable() override;
+	UE_API bool IsReadyForReuse() override;
+	UE_API void ShutdownPoolable() override;
 
 private:
-	uint32 GetConverterInfoFlags() const override;
-	bool Convert(FRHICommandListImmediate& RHICmdList, FTextureRHIRef& InDstTexture, const FConversionHints& Hints) override;
+	UE_API uint32 GetConverterInfoFlags() const override;
+	UE_API bool Convert(FRHICommandListImmediate& RHICmdList, FTextureRHIRef& InDstTexture, const FConversionHints& Hints) override;
 
 	friend class FElectraTextureSampleSupport;
 
-	FTextureRHIRef InitializeTextureOES(AHardwareBuffer* HardwareBuffer);
-	FTextureRHIRef InitializeTextureVulkan(AHardwareBuffer* HardwareBuffer);
+	UE_API FTextureRHIRef InitializeTextureOES(AHardwareBuffer* HardwareBuffer);
+	UE_API FTextureRHIRef InitializeTextureVulkan(AHardwareBuffer* HardwareBuffer);
 
-	void InitializeTexture(EPixelFormat PixelFormat);
-	void SetupFromBuffer(const void* InBuffer, int32 InBufferSize);
-	void SetImageResources(jobject InImageResources);
-	void CleanupImageResources();
-	void CopyFromExternalTextureOES(FRHICommandListImmediate& RHICmdList, FTextureRHIRef& InDstTexture, FTextureRHIRef& InSrcTexture, const FVector2f& Scale, const FVector2f& Offset);
-	void CopyFromExternalTextureVulkan(FRHICommandListImmediate& RHICmdList, FTextureRHIRef& InDstTexture, FTextureRHIRef& InSrcTexture, const FVector2f& InScale, const FVector2f& InOffset);
+	UE_API void InitializeTexture(EPixelFormat PixelFormat);
+	UE_API void SetupFromBuffer(const void* InBuffer, int32 InBufferSize);
+	UE_API void SetImageResources(jobject InImageResources);
+	UE_API void CleanupImageResources();
+	UE_API void CopyFromExternalTextureOES(FRHICommandListImmediate& RHICmdList, FTextureRHIRef& InDstTexture, FTextureRHIRef& InSrcTexture, const FVector2f& Scale, const FVector2f& Offset);
+	UE_API void CopyFromExternalTextureVulkan(FRHICommandListImmediate& RHICmdList, FTextureRHIRef& InDstTexture, FTextureRHIRef& InSrcTexture, const FVector2f& InScale, const FVector2f& InOffset);
 
-	bool ConvertCpuOutputPath(FRHICommandListImmediate& RHICmdList, FTextureRHIRef& InDstTexture, const FConversionHints& Hints);
-	bool ConvertGpuOutputPath(FRHICommandListImmediate& RHICmdList, FTextureRHIRef& InDstTexture, const FConversionHints& Hints);
+	UE_API bool ConvertCpuOutputPath(FRHICommandListImmediate& RHICmdList, FTextureRHIRef& InDstTexture, const FConversionHints& Hints);
+	UE_API bool ConvertGpuOutputPath(FRHICommandListImmediate& RHICmdList, FTextureRHIRef& InDstTexture, const FConversionHints& Hints);
 
 
 	TSharedPtr<FElectraTextureSampleSupport, ESPMode::ThreadSafe> Support;
diff --git a/Engine/Plugins/Runtime/InstancedActors/Source/InstancedActors/Public/InstancedActorsManager.h b/Engine/Plugins/Runtime/InstancedActors/Source/InstancedActors/Public/InstancedActorsManager.h
--- a/Engine/Plugins/Runtime/InstancedActors/Source/InstancedActors/Public/InstancedActorsManager.h
+++ b/Engine/Plugins/Runtime/InstancedActors/Source/InstancedActors/Public/InstancedActorsManager.h
@@ -259,7 +259,7 @@ protected:
 	// mz@todo a temp glue to be removed once InstancedActors get moved out of original project for good. It's virtual and I don't like it.
 	virtual void RequestActorSave(AActor* Actor) {}
 
-	UE_API virtual void OnInitializeModifyAndSpawnEntities() {}
+	virtual void OnInitializeModifyAndSpawnEntities() {}
 
 	//~ Begin AActor Overrides
 #if UE_WITH_IRIS
diff --git a/Engine/Source/Runtime/AudioMixerCore/Public/AudioMixer.h b/Engine/Source/Runtime/AudioMixerCore/Public/AudioMixer.h
--- a/Engine/Source/Runtime/AudioMixerCore/Public/AudioMixer.h
+++ b/Engine/Source/Runtime/AudioMixerCore/Public/AudioMixer.h
@@ -815,8 +815,8 @@ namespace Audio
 
 		//~ Begin IAudioMixerPlatformInterface
 		UE_API virtual bool RequestDeviceSwap(const FString& DeviceID, const bool bInForce, const TCHAR* InReason) override;
-		UE_API virtual void LockDeviceSwap() override { DeviceSwapCriticalSection.Lock(); }
-		UE_API virtual void UnlockDeviceSwap() override { DeviceSwapCriticalSection.Unlock(); }
+		virtual void LockDeviceSwap() override { DeviceSwapCriticalSection.Lock(); }
+		virtual void UnlockDeviceSwap() override { DeviceSwapCriticalSection.Unlock(); }
 		UE_API virtual bool CheckAudioDeviceChange() override;
 		UE_API virtual bool MoveAudioStreamToNewAudioDevice() override;
 		UE_API virtual void ResumePlaybackOnNewDevice() override;
diff --git a/Engine/Source/Runtime/Engine/Classes/Engine/World.h b/Engine/Source/Runtime/Engine/Classes/Engine/World.h
--- a/Engine/Source/Runtime/Engine/Classes/Engine/World.h
+++ b/Engine/Source/Runtime/Engine/Classes/Engine/World.h
@@ -4044,7 +4044,7 @@ public:
 
 	/** Returns whether script is executing within a vulkan preview window */
 	UE_DEPRECATED(5.7, "PlayInVulkanPreview is no longer supported, please use PlayInMobilePreview instead")
-	UE_API bool IsPlayInVulkanPreview() const { return false; };
+	bool IsPlayInVulkanPreview() const { return false; };
 
 	/** Returns true if this world is any kind of game world (including PIE worlds) */
 	UE_API bool IsGameWorld() const;
