ElectraDecoders: Passing any low level D3D12 failures up for better error reporting

#rnx

diff --git a/Engine/Plugins/Media/AppleProResDecoderElectra/Source/AppleProResDecoderElectra/Private/ProResDecoder/Windows/ElectraMediaProResDecoder.cpp b/Engine/Plugins/Media/AppleProResDecoderElectra/Source/AppleProResDecoderElectra/Private/ProResDecoder/Windows/ElectraMediaProResDecoder.cpp
--- a/Engine/Plugins/Media/AppleProResDecoderElectra/Source/AppleProResDecoderElectra/Private/ProResDecoder/Windows/ElectraMediaProResDecoder.cpp
+++ b/Engine/Plugins/Media/AppleProResDecoderElectra/Source/AppleProResDecoderElectra/Private/ProResDecoder/Windows/ElectraMediaProResDecoder.cpp
@@ -604,23 +604,35 @@ bool FVideoDecoderProResElectra::ConfigureOutputForInputFormat()
 	switch(CurrentHeader.GetSourcePixFormat())
 	{
 		case FProResHeader::ESourcePixFormat::Twovuy:
+		{
 			OutputPixelFormat = kPRFormat_2vuy;
 			break;
+		}
 		case FProResHeader::ESourcePixFormat::v210:
+		{
 			OutputPixelFormat = kPRFormat_v210;
 			break;
+		}
 		case FProResHeader::ESourcePixFormat::v216:
+		{
 			OutputPixelFormat = kPRFormat_v216;
 			break;
+		}
 		case FProResHeader::ESourcePixFormat::r4fl:
+		{
 			OutputPixelFormat = kPRFormat_r4fl;
 			break;
+		}
 		case FProResHeader::ESourcePixFormat::b64a:
+		{
 			OutputPixelFormat = kPRFormat_b64a;
 			break;
+		}
 		case FProResHeader::ESourcePixFormat::R10k:
+		{
 			OutputPixelFormat = kPRFormat_R10k;
 			break;
+		}
 		case FProResHeader::ESourcePixFormat::r408:
 		case FProResHeader::ESourcePixFormat::v408:
 		case FProResHeader::ESourcePixFormat::RGB:
@@ -628,8 +640,10 @@ bool FVideoDecoderProResElectra::ConfigureOutputForInputFormat()
 		case FProResHeader::ESourcePixFormat::n302:
 		case FProResHeader::ESourcePixFormat::l302:
 		default:
+		{
 			OutputPixelFormat = kPRFormat_y416;
 			break;
+		}
 	}
 
 	//FProResHeader::EChromaType = CurrentHeader.GetChromaType();
@@ -811,6 +825,8 @@ IElectraDecoder::EDecoderError FVideoDecoderProResElectra::DecodeAccessUnit(cons
 			// (the memory is WC configured, so the decoder must not read from it -- this seems to be the case with ProRes)
 
 			TRefCountPtr D3D12Device(static_cast<ID3D12Device*>(PlatformDevice));
+			HRESULT Result;
+			FString ResultMsg;
 
 			// Create the resource pool as needed...
 			if (!D3D12ResourcePool)
@@ -821,15 +837,15 @@ IElectraDecoder::EDecoderError FVideoDecoderProResElectra::DecodeAccessUnit(cons
 
 			// Request resource and fence...
 			uint32 BufferPitch;
-			D3D12ResourcePool->AllocateOutputDataAsBuffer(NewOutput->GPUBuffer, BufferPitch);
+			D3D12ResourcePool->AllocateOutputDataAsBuffer(Result, ResultMsg, NewOutput->GPUBuffer, BufferPitch);
 			if (NewOutput->GPUBuffer.Resource)
 			{
 				// Correct output pitch to what the resource is setup for
 				NewOutput->Pitch = BufferPitch;
 
 				void* BufferAddr;
-				HRESULT Res = NewOutput->GPUBuffer.Resource->Map(0, nullptr, &BufferAddr);
-				if (SUCCEEDED(Res))
+				Result = NewOutput->GPUBuffer.Resource->Map(0, nullptr, &BufferAddr);
+				if (SUCCEEDED(Result))
 				{
 					check((uint32)OutputBufferBytesPerRow <= BufferPitch);
 					pb.baseAddr = (unsigned char*)BufferAddr;
@@ -837,13 +853,13 @@ IElectraDecoder::EDecoderError FVideoDecoderProResElectra::DecodeAccessUnit(cons
 				}
 				else
 				{
-					PostError(0, TEXT("Failed to map DX12 texture buffer"), ERRCODE_INTERNAL_FAILED_TO_DECODE_INPUT);
+					PostError((int32)Result, TEXT("Failed to map DX12 texture buffer"), ERRCODE_INTERNAL_FAILED_TO_DECODE_INPUT);
 					return IElectraDecoder::EDecoderError::Error;
 				}
 			}
 			else
 			{
-				PostError(0, TEXT("Failed to allocate DX12 texture buffer"), ERRCODE_INTERNAL_FAILED_TO_DECODE_INPUT);
+				PostError((int32)Result, FString::Printf(TEXT("Failed to allocate DX12 texture buffer (%s)"), *ResultMsg), ERRCODE_INTERNAL_FAILED_TO_DECODE_INPUT);
 				return IElectraDecoder::EDecoderError::Error;
 			}
 		}
diff --git a/Engine/Plugins/Media/AvidDNxHDDecoderElectra/Source/Private/Windows/ElectraMediaAvidDNxHDDecoder.cpp b/Engine/Plugins/Media/AvidDNxHDDecoderElectra/Source/Private/Windows/ElectraMediaAvidDNxHDDecoder.cpp
--- a/Engine/Plugins/Media/AvidDNxHDDecoderElectra/Source/Private/Windows/ElectraMediaAvidDNxHDDecoder.cpp
+++ b/Engine/Plugins/Media/AvidDNxHDDecoderElectra/Source/Private/Windows/ElectraMediaAvidDNxHDDecoder.cpp
@@ -641,28 +641,28 @@ IElectraDecoder::EDecoderError FVideoDecoderAvidDNxHDElectra::DecodeAccessUnit(c
 						NewOutput->Pitch *= 4;
 						break;
 					}
-					case	DNX_CCO_RGBA_Interleaved:
+					case DNX_CCO_RGBA_Interleaved:
 					{
 						NewOutput->BufferFormat = EPixelFormat::PF_R16G16B16A16_UNORM;
 						NewOutput->BufferEncoding = EElectraTextureSamplePixelEncoding::Native;
 						NewOutput->Pitch *= 8;
 						break;
 					}
-					case	DNX_CCO_ABGR_Interleaved:
+					case DNX_CCO_ABGR_Interleaved:
 					{
 						NewOutput->BufferFormat = EPixelFormat::PF_A16B16G16R16;
 						NewOutput->BufferEncoding = EElectraTextureSamplePixelEncoding::Native;
 						NewOutput->Pitch *= 8;
 						break;
 					}
-					case	DNX_CCO_CbYCrA_Interleaved:
+					case DNX_CCO_CbYCrA_Interleaved:
 					{
 						NewOutput->BufferFormat = EPixelFormat::PF_A16B16G16R16;
 						NewOutput->BufferEncoding = EElectraTextureSamplePixelEncoding::YCbCr_Alpha;
 						NewOutput->Pitch *= 8;
 						break;
 					}
-					case	DNX_CCO_YCbCr_Planar:
+					case DNX_CCO_YCbCr_Planar:
 					{
 						NewOutput->BufferFormat = EPixelFormat::PF_P010;
 						NewOutput->BufferEncoding = EElectraTextureSamplePixelEncoding::Native;
@@ -701,8 +701,8 @@ IElectraDecoder::EDecoderError FVideoDecoderAvidDNxHDElectra::DecodeAccessUnit(c
 
 		uint32 BPP = NewOutput->Pitch / NewOutput->Width;
 
-		uint8* OutBufferData;
-		uint32 OutBufferSize;
+		uint8* OutBufferData = nullptr;
+		uint32 OutBufferSize = 0;
 		TArray<uint8> TempBuffer;
 
 #if ELECTRA_MEDIAGPUBUFFER_DX12
@@ -710,6 +710,8 @@ IElectraDecoder::EDecoderError FVideoDecoderAvidDNxHDElectra::DecodeAccessUnit(c
 		if (bUseGPUBuffers)
 		{
 			TRefCountPtr D3D12Device(static_cast<ID3D12Device*>(PlatformDevice));
+			HRESULT HResult;
+			FString HResultMsg;
 
 			uint32 AllocHeight = bIsPlanar ? (NewOutput->Height * 3 / 2) : NewOutput->Height;
 
@@ -727,29 +729,39 @@ IElectraDecoder::EDecoderError FVideoDecoderAvidDNxHDElectra::DecodeAccessUnit(c
 			}
 
 			// Request resource and fence...
-			D3D12ResourcePool->AllocateOutputDataAsBuffer(NewOutput->GPUBuffer, ResourcePitch);
-
-			NewOutput->Pitch = ResourcePitch;
+			if (D3D12ResourcePool->AllocateOutputDataAsBuffer(HResult, HResultMsg, NewOutput->GPUBuffer, ResourcePitch))
+			{
+				NewOutput->Pitch = ResourcePitch;
 
-			uint32 OutputPitch = BPP * NewOutput->Width;
+				uint32 OutputPitch = BPP * NewOutput->Width;
 
-			// Pitch compatible with DX12 resource?
-			if (OutputPitch != ResourcePitch)
-			{
-				// No, decode to temp buffer...
-				TempBuffer.AddUninitialized(OutputPitch * AllocHeight);
+				// Pitch compatible with DX12 resource?
+				if (OutputPitch != ResourcePitch)
+				{
+					// No, decode to temp buffer...
+					TempBuffer.AddUninitialized(OutputPitch * AllocHeight);
 
-				OutBufferData = TempBuffer.GetData();
-				OutBufferSize = TempBuffer.Num();
+					OutBufferData = TempBuffer.GetData();
+					OutBufferSize = TempBuffer.Num();
 
-				check(OutBufferSize <= AllocHeight * ResourcePitch);
+					check(OutBufferSize <= AllocHeight * ResourcePitch);
+				}
+				else
+				{
+					// Yes, decode directly into the DX12 resource...
+					HResult = NewOutput->GPUBuffer.Resource->Map(0, nullptr, (void**)&OutBufferData);
+					if (FAILED(HResult))
+					{
+						PostError((int32)HResult, FString::Printf(TEXT("Failed to map D3D12 buffer")), ERRCODE_INTERNAL_FAILED_TO_DECODE_INPUT);
+						return IElectraDecoder::EDecoderError::Error;
+					}
+					OutBufferSize = ResourcePitch * AllocHeight;
+				}
 			}
 			else
 			{
-				// Yes, decode directly into the DX12 resource...
-				HRESULT Res = NewOutput->GPUBuffer.Resource->Map(0, nullptr, (void**)&OutBufferData);
-				check(SUCCEEDED(Res));
-				OutBufferSize = ResourcePitch * AllocHeight;
+				PostError((int32)HResult, FString::Printf(TEXT("Failed to allocate D3D12 buffer (%s)"), *HResultMsg), ERRCODE_INTERNAL_FAILED_TO_DECODE_INPUT);
+				return IElectraDecoder::EDecoderError::Error;
 			}
 		}
 		else
@@ -962,18 +974,18 @@ bool FVideoDecoderAvidDNxHDElectra::FDecoderHandle::SetupDecompressStruct()
 				CurrentUncompressedParams.compOrder = DecodeOperationParams.decodeAlpha ? DNX_ColorComponentOrder_t::DNX_CCO_CbYCrY_A : DNX_ColorComponentOrder_t::DNX_CCO_YCbYCr_NoA;
 				break;
 			}
-		case 12:
+			case 12:
+			{
+				CurrentUncompressedParams.compType = DNX_ComponentType_t::DNX_CT_USHORT_12_4;
+				// FIXME: What output format would the alpha channel require us to use?
+				if (DecodeOperationParams.decodeAlpha)
 				{
-					CurrentUncompressedParams.compType = DNX_ComponentType_t::DNX_CT_USHORT_12_4;
-					// FIXME: What output format would the alpha channel require us to use?
-					if (DecodeOperationParams.decodeAlpha)
-					{
-						UE_LOG(LogAvidDNxHDElectraDecoder, Log, TEXT("Compressed alpha with YCbYCr is not currently handled"));
-						return false;
-					}
-					CurrentUncompressedParams.compOrder = DecodeOperationParams.decodeAlpha ? DNX_ColorComponentOrder_t::DNX_CCO_CbYCrY_A : DNX_ColorComponentOrder_t::DNX_CCO_YCbYCr_NoA;
-					break;
+					UE_LOG(LogAvidDNxHDElectraDecoder, Log, TEXT("Compressed alpha with YCbYCr is not currently handled"));
+					return false;
 				}
+				CurrentUncompressedParams.compOrder = DecodeOperationParams.decodeAlpha ? DNX_ColorComponentOrder_t::DNX_CCO_CbYCrY_A : DNX_ColorComponentOrder_t::DNX_CCO_YCbYCr_NoA;
+				break;
+			}
 			default:
 			{
 				UE_LOG(LogAvidDNxHDElectraDecoder, Log, TEXT("Compressed %d bit depth is not currently handled"), CurrentCompressedParams.depth);
diff --git a/Engine/Plugins/Media/ElectraCodecs/Source/ElectraDecoders/Private/DX/VideoDecoderH264_DX.cpp b/Engine/Plugins/Media/ElectraCodecs/Source/ElectraDecoders/Private/DX/VideoDecoderH264_DX.cpp
--- a/Engine/Plugins/Media/ElectraCodecs/Source/ElectraDecoders/Private/DX/VideoDecoderH264_DX.cpp
+++ b/Engine/Plugins/Media/ElectraCodecs/Source/ElectraDecoders/Private/DX/VideoDecoderH264_DX.cpp
@@ -1018,6 +1018,7 @@ bool FElectraVideoDecoderH264_DX::ConvertDecoderOutput()
 	}
 
 	HRESULT Result;
+	FString ResultMsg;
 	LONGLONG llTimeStamp = 0;
 	VERIFY_HR(DecodedOutputSample->GetSampleTime(&llTimeStamp), TEXT("Failed to get video decoder output sample timestamp"), ERRCODE_INTERNAL_FAILED_TO_CONVERT_OUTPUT_SAMPLE);
 	FTimespan TimeStamp((int64) llTimeStamp);
@@ -1153,15 +1154,18 @@ bool FElectraVideoDecoderH264_DX::ConvertDecoderOutput()
 		DWORD BuffersNum = 0;
 		if ((Result = DecodedOutputSample->GetBufferCount(&BuffersNum)) != S_OK)
 		{
+			PostError(Result, TEXT("Failed to get buffer count"), ERRCODE_INTERNAL_FAILED_TO_CONVERT_OUTPUT_SAMPLE);
 			return false;
 		}
 		if (BuffersNum != 1)
 		{
+			PostError(0, TEXT("Invalid buffer count"), ERRCODE_INTERNAL_FAILED_TO_CONVERT_OUTPUT_SAMPLE);
 			return false;
 		}
 		TRefCountPtr<IMFMediaBuffer> Buffer;
 		if ((Result = DecodedOutputSample->GetBufferByIndex(0, Buffer.GetInitReference())) != S_OK)
 		{
+			PostError(Result, TEXT("Failed to get buffer by index"), ERRCODE_INTERNAL_FAILED_TO_CONVERT_OUTPUT_SAMPLE);
 			return false;
 		}
 
@@ -1187,13 +1191,15 @@ bool FElectraVideoDecoderH264_DX::ConvertDecoderOutput()
 		TRefCountPtr<IMF2DBuffer> Buffer2D;
 		if ((Result = Buffer->QueryInterface(__uuidof(IMF2DBuffer), (void**)Buffer2D.GetInitReference())) != S_OK)
 		{
+			PostError(Result, TEXT("Failed to get 2d buffer from buffer"), ERRCODE_INTERNAL_FAILED_TO_CONVERT_OUTPUT_SAMPLE);
 			return false;
 		}
 
 		// Request resource and fence...
 		uint32 BufferPitch;
-		if (!D3D12ResourcePool->AllocateOutputDataAsBuffer(NewOutput->Buffer, BufferPitch))
+		if (!D3D12ResourcePool->AllocateOutputDataAsBuffer(Result, ResultMsg, NewOutput->Buffer, BufferPitch))
 		{
+			PostError(Result, FString::Printf(TEXT("Failed to allocate D3D12 data buffer (%s)"), *ResultMsg), ERRCODE_INTERNAL_FAILED_TO_CONVERT_OUTPUT_SAMPLE);
 			return false;
 		}
 
diff --git a/Engine/Plugins/Media/ElectraCodecs/Source/ElectraDecoders/Private/DX/VideoDecoderH265_DX.cpp b/Engine/Plugins/Media/ElectraCodecs/Source/ElectraDecoders/Private/DX/VideoDecoderH265_DX.cpp
--- a/Engine/Plugins/Media/ElectraCodecs/Source/ElectraDecoders/Private/DX/VideoDecoderH265_DX.cpp
+++ b/Engine/Plugins/Media/ElectraCodecs/Source/ElectraDecoders/Private/DX/VideoDecoderH265_DX.cpp
@@ -1003,6 +1003,7 @@ bool FElectraVideoDecoderH265_DX::ConvertDecoderOutput()
 	}
 
 	HRESULT Result;
+	FString ResultMsg;
 	LONGLONG llTimeStamp = 0;
 	VERIFY_HR(DecodedOutputSample->GetSampleTime(&llTimeStamp), TEXT("Failed to get video decoder output sample timestamp"), ERRCODE_INTERNAL_FAILED_TO_CONVERT_OUTPUT_SAMPLE);
 	FTimespan TimeStamp((int64) llTimeStamp);
@@ -1159,15 +1160,18 @@ bool FElectraVideoDecoderH265_DX::ConvertDecoderOutput()
 		DWORD BuffersNum = 0;
 		if ((Result = DecodedOutputSample->GetBufferCount(&BuffersNum)) != S_OK)
 		{
+			PostError(Result, TEXT("Failed to get buffer count"), ERRCODE_INTERNAL_FAILED_TO_CONVERT_OUTPUT_SAMPLE);
 			return false;
 		}
 		if (BuffersNum != 1)
 		{
+			PostError(0, TEXT("Invalid buffer count"), ERRCODE_INTERNAL_FAILED_TO_CONVERT_OUTPUT_SAMPLE);
 			return false;
 		}
 		TRefCountPtr<IMFMediaBuffer> Buffer;
 		if ((Result = DecodedOutputSample->GetBufferByIndex(0, Buffer.GetInitReference())) != S_OK)
 		{
+			PostError(Result, TEXT("Failed to get buffer by index"), ERRCODE_INTERNAL_FAILED_TO_CONVERT_OUTPUT_SAMPLE);
 			return false;
 		}
 
@@ -1193,13 +1197,15 @@ bool FElectraVideoDecoderH265_DX::ConvertDecoderOutput()
 		TRefCountPtr<IMF2DBuffer> Buffer2D;
 		if ((Result = Buffer->QueryInterface(__uuidof(IMF2DBuffer), (void**)Buffer2D.GetInitReference())) != S_OK)
 		{
+			PostError(Result, TEXT("Failed to get 2d buffer from buffer"), ERRCODE_INTERNAL_FAILED_TO_CONVERT_OUTPUT_SAMPLE);
 			return false;
 		}
 
 		// Request resource and fence...
 		uint32 BufferPitch;
-		if (!D3D12ResourcePool->AllocateOutputDataAsBuffer(NewOutput->Buffer, BufferPitch))
+		if (!D3D12ResourcePool->AllocateOutputDataAsBuffer(Result, ResultMsg, NewOutput->Buffer, BufferPitch))
 		{
+			PostError(Result, FString::Printf(TEXT("Failed to allocate D3D12 data buffer (%s)"), *ResultMsg), ERRCODE_INTERNAL_FAILED_TO_CONVERT_OUTPUT_SAMPLE);
 			return false;
 		}
 
diff --git a/Engine/Plugins/Media/ElectraCodecs/Source/ElectraDecoders/Private/Windows/WindowsElectraDecodersPlatformResources.cpp b/Engine/Plugins/Media/ElectraCodecs/Source/ElectraDecoders/Private/Windows/WindowsElectraDecodersPlatformResources.cpp
--- a/Engine/Plugins/Media/ElectraCodecs/Source/ElectraDecoders/Private/Windows/WindowsElectraDecodersPlatformResources.cpp
+++ b/Engine/Plugins/Media/ElectraCodecs/Source/ElectraDecoders/Private/Windows/WindowsElectraDecodersPlatformResources.cpp
@@ -172,6 +172,7 @@ static bool SetupOutputTextureSampleFromDX12Resource(FString& OutErrorMessage, T
 													 uint32 MaxWidth, uint32 MaxHeight, uint32 MaxOutputBuffers)
 {
 	HRESULT Result;
+	FString ResultMsg;
 
 	// Check that the sample is a fresh or a cleanly released one.
 	check(!InOutTextureSampleToSetup->D3DFence.IsValid());
@@ -333,9 +334,9 @@ static bool SetupOutputTextureSampleFromDX12Resource(FString& OutErrorMessage, T
 		}
 
 		FElectraMediaDecoderOutputBufferPool_DX12::FOutputData OutputData;
-		if (!InOutTextureSampleToSetup->D3D12ResourcePool->AllocateOutputDataAsTexture(OutputData, TextureDim.X, TextureDim.Y, DXGIFmt))
+		if (!InOutTextureSampleToSetup->D3D12ResourcePool->AllocateOutputDataAsTexture(Result, ResultMsg, OutputData, TextureDim.X, TextureDim.Y, DXGIFmt))
 		{
-			OutErrorMessage = FString::Printf(TEXT("SetupOutputTextureSampleFromDX12Resource(): Could not allocate texture resource from heap"));
+			OutErrorMessage = FString::Printf(TEXT("SetupOutputTextureSampleFromDX12Resource(): Could not allocate texture resource from heap (%x, %s)"), Result, *ResultMsg);
 			return false;
 		}
 
diff --git a/Engine/Plugins/Media/ElectraUtil/Source/ElectraSamples/Public/Windows/WindowsElectraTextureSampleGPUBufferHelper.h b/Engine/Plugins/Media/ElectraUtil/Source/ElectraSamples/Public/Windows/WindowsElectraTextureSampleGPUBufferHelper.h
--- a/Engine/Plugins/Media/ElectraUtil/Source/ElectraSamples/Public/Windows/WindowsElectraTextureSampleGPUBufferHelper.h
+++ b/Engine/Plugins/Media/ElectraUtil/Source/ElectraSamples/Public/Windows/WindowsElectraTextureSampleGPUBufferHelper.h
@@ -177,9 +177,9 @@ public:
 	}
 
 	// Allocate a buffer resource and return suitable sync fence data
-	TRefCountPtr<ID3D12Resource> AllocateOutputDataAsBuffer(TWeakPtr<FElectraMediaDecoderOutputBufferPool_DX12> InPoolWeakRef, uint32& OutPitch)
+	TRefCountPtr<ID3D12Resource> AllocateOutputDataAsBuffer(HRESULT& OutResult, FString& OutMessage, TWeakPtr<FElectraMediaDecoderOutputBufferPool_DX12> InPoolWeakRef, uint32& OutPitch)
 	{
-		TRefCountPtr<ID3D12Resource> Resource = AllocateBuffer(InPoolWeakRef, OutPitch, [NumBytes = BufferSize](D3D12_RESOURCE_DESC& Desc)
+		TRefCountPtr<ID3D12Resource> Resource = AllocateBuffer(OutResult, OutMessage, InPoolWeakRef, OutPitch, [NumBytes = BufferSize](D3D12_RESOURCE_DESC& Desc)
 		{
 			Desc.MipLevels = 1;
 			Desc.Format = DXGI_FORMAT_UNKNOWN;
@@ -196,10 +196,10 @@ public:
 	}
 
 	// Allocate a texture resource and return suitable sync fence data
-	TRefCountPtr<ID3D12Resource> AllocateOutputDataAsTexture(TWeakPtr<FElectraMediaDecoderOutputBufferPool_DX12> InPoolWeakRef, uint32 InWidth, uint32 InHeight, DXGI_FORMAT InPixFmt)
+	TRefCountPtr<ID3D12Resource> AllocateOutputDataAsTexture(HRESULT& OutResult, FString& OutMessage, TWeakPtr<FElectraMediaDecoderOutputBufferPool_DX12> InPoolWeakRef, uint32 InWidth, uint32 InHeight, DXGI_FORMAT InPixFmt)
 	{
 		uint32 OutPitch;
-		TRefCountPtr<ID3D12Resource> Resource = AllocateBuffer(InPoolWeakRef, OutPitch, [InWidth, InHeight, InPixFmt](D3D12_RESOURCE_DESC& Desc)
+		TRefCountPtr<ID3D12Resource> Resource = AllocateBuffer(OutResult, OutMessage, InPoolWeakRef, OutPitch, [InWidth, InHeight, InPixFmt](D3D12_RESOURCE_DESC& Desc)
 		{
 			Desc.MipLevels = 1;
 			Desc.Format = InPixFmt;
@@ -236,9 +236,10 @@ private:
 #endif
 	}
 
-	TRefCountPtr<ID3D12Resource> AllocateBuffer(TWeakPtr<FElectraMediaDecoderOutputBufferPool_DX12> InPoolWeakRef, uint32& OutBufferPitch, TFunction<void(D3D12_RESOURCE_DESC& Desc)> && InitializeDesc)
+	TRefCountPtr<ID3D12Resource> AllocateBuffer(HRESULT& OutResult, FString& OutMessage, TWeakPtr<FElectraMediaDecoderOutputBufferPool_DX12> InPoolWeakRef, uint32& OutBufferPitch, TFunction<void(D3D12_RESOURCE_DESC& Desc)> && InitializeDesc)
 	{
 		TRefCountPtr<ID3D12Resource> Resource;
+		OutResult = S_OK;
 
 		// Find the buffer index we can use...
 		int32 BufferIdx, OldFreeMask, NewFreeMask;
@@ -260,8 +261,8 @@ private:
 			// Make a placed resource to use...
 			D3D12_RESOURCE_DESC Desc = {};
 			InitializeDesc(Desc);
-			HRESULT Res = D3D12Device->CreatePlacedResource(D3D12OutputHeap, BufferSize * BufferIdx, &Desc, (D3D12HeapType == D3D12_HEAP_TYPE_UPLOAD) ? D3D12_RESOURCE_STATE_GENERIC_READ : D3D12_RESOURCE_STATE_COMMON, nullptr, IID_PPV_ARGS(Resource.GetInitReference()));
-			if (Res == S_OK)
+			OutResult = D3D12Device->CreatePlacedResource(D3D12OutputHeap, BufferSize * BufferIdx, &Desc, (D3D12HeapType == D3D12_HEAP_TYPE_UPLOAD) ? D3D12_RESOURCE_STATE_GENERIC_READ : D3D12_RESOURCE_STATE_COMMON, nullptr, IID_PPV_ARGS(Resource.GetInitReference()));
+			if (OutResult == S_OK)
 			{
 #if !UE_BUILD_SHIPPING
 				Resource->SetName(TEXT("ElectraOutputBufferPoolBufferResource"));
@@ -271,26 +272,29 @@ private:
 
 				// Register a destruction callback so we can reset our memory management here transparently...
 				TRefCountPtr<ID3DDestructionNotifier> Notifier;
-				Res = Resource->QueryInterface(__uuidof(ID3DDestructionNotifier), (void**)Notifier.GetInitReference());
-				if (SUCCEEDED(Res))
+				OutResult = Resource->QueryInterface(__uuidof(ID3DDestructionNotifier), (void**)Notifier.GetInitReference());
+				if (SUCCEEDED(OutResult))
 				{
 					// note: we keep a reference to the heap in our context data for the destruction callback, so we can ensure the heap lives as long as there are placed resources in it
 					UINT CallbackID;
-					Res = Notifier->RegisterDestructionCallback(ResourceDestructionCallback, new FResourceInfo(InPoolWeakRef, BlockIdx, BufferIdx, D3D12OutputHeap), &CallbackID);
-					if (!(SUCCEEDED(Res)))
+					OutResult = Notifier->RegisterDestructionCallback(ResourceDestructionCallback, new FResourceInfo(InPoolWeakRef, BlockIdx, BufferIdx, D3D12OutputHeap), &CallbackID);
+					if (!(SUCCEEDED(OutResult)))
 					{
 						Resource = nullptr;
 						FreeBuffer(BufferIdx);
+						OutMessage = TEXT("Call to RegisterDestructionCallback() failed");
 					}
 				}
 				else
 				{
 					Resource = nullptr;
 					FreeBuffer(BufferIdx);
+					OutMessage = TEXT("Call to QueryInterface(ID3DDestructionNotifier) failed");
 				}
 			}
 			else
 			{
+				OutMessage = TEXT("Call to CreatePlacedResource() failed");
 				FreeBuffer(BufferIdx);
 			}
 		}
@@ -416,17 +420,18 @@ public:
 		}
 	};
 
-	bool AllocateOutputDataAsBuffer(FOutputData& OutData, uint32& OutPitch)
+	bool AllocateOutputDataAsBuffer(HRESULT& OutResult, FString& OutMessage, FOutputData& OutData, uint32& OutPitch)
 	{
 		check(BytesPerPixel != 0);
 
+		OutResult = S_OK;
 		OutData.Resource = nullptr;
 		for(FElectraMediaDecoderOutputBufferPoolBlock_DX12* Block : Blocks)
 		{
 			if (Block->BufferAvailable())
 			{
-				OutData.Resource = Block->AllocateOutputDataAsBuffer(AsWeak(), OutPitch);
-				if (OutData.Resource)
+				OutData.Resource = Block->AllocateOutputDataAsBuffer(OutResult, OutMessage, AsWeak(), OutPitch);
+				if (OutResult != S_OK || OutData.Resource)
 				{
 					break;
 				}
@@ -434,10 +439,11 @@ public:
 		}
 
 		const uint32 NumAllocatedSlots = ((Blocks.Num() - 1) * 32U) + Blocks.Last()->GetMaxNumBuffers();
-		if (!OutData.Resource.IsValid() && NumAllocatedSlots < MaxNumBuffers)
+		// Allocate a new block if we could not get a buffer, but only when there is no error yet.
+		if (!OutData.Resource.IsValid() && OutResult == S_OK && NumAllocatedSlots < MaxNumBuffers)
 		{
 			Blocks.Emplace(new FElectraMediaDecoderOutputBufferPoolBlock_DX12(Blocks.Num(), D3D12Device, MaxNumBuffers-NumAllocatedSlots > 32 ? 32 : MaxNumBuffers-NumAllocatedSlots, Width, Height, BytesPerPixel, D3D12HeapType));
-			OutData.Resource = Blocks.Last()->AllocateOutputDataAsBuffer(AsWeak(), OutPitch);
+			OutData.Resource = Blocks.Last()->AllocateOutputDataAsBuffer(OutResult, OutMessage, AsWeak(), OutPitch);
 		}
 
 		if (OutData.Resource.IsValid())
@@ -448,17 +454,18 @@ public:
 		return false;
 	}
 
-	bool AllocateOutputDataAsTexture(FOutputData& OutData, uint32 InWidth, uint32 InHeight, DXGI_FORMAT InPixFmt)
+	bool AllocateOutputDataAsTexture(HRESULT& OutResult, FString& OutMessage, FOutputData& OutData, uint32 InWidth, uint32 InHeight, DXGI_FORMAT InPixFmt)
 	{
 		check(PixFmt != DXGI_FORMAT_UNKNOWN);
 
 		OutData.Resource = nullptr;
+		OutResult = S_OK;
 		for(FElectraMediaDecoderOutputBufferPoolBlock_DX12* Block : Blocks)
 		{
 			if (Block->BufferAvailable())
 			{
-				OutData.Resource = Block->AllocateOutputDataAsTexture(AsWeak(), InWidth, InHeight, InPixFmt);
-				if (OutData.Resource)
+				OutData.Resource = Block->AllocateOutputDataAsTexture(OutResult, OutMessage, AsWeak(), InWidth, InHeight, InPixFmt);
+				if (OutResult != S_OK || OutData.Resource)
 				{
 					break;
 				}
@@ -466,10 +473,11 @@ public:
 		}
 
 		const uint32 NumAllocatedSlots = ((Blocks.Num() - 1) * 32U) + Blocks.Last()->GetMaxNumBuffers();
-		if (!OutData.Resource.IsValid() && NumAllocatedSlots < MaxNumBuffers)
+		// Allocate a new block if we could not get a buffer, but only when there is no error yet.
+		if (!OutData.Resource.IsValid() && OutResult == S_OK && NumAllocatedSlots < MaxNumBuffers)
 		{
 			Blocks.Emplace(new FElectraMediaDecoderOutputBufferPoolBlock_DX12(Blocks.Num(), D3D12Device, MaxNumBuffers-NumAllocatedSlots > 32 ? 32 : MaxNumBuffers-NumAllocatedSlots, Width, Height, PixFmt, D3D12HeapType));
-			OutData.Resource = Blocks.Last()->AllocateOutputDataAsTexture(AsWeak(), InWidth, InHeight, InPixFmt);
+			OutData.Resource = Blocks.Last()->AllocateOutputDataAsTexture(OutResult, OutMessage, AsWeak(), InWidth, InHeight, InPixFmt);
 		}
 
 		if (OutData.Resource.IsValid())
diff --git a/Engine/Plugins/Media/HAPDecoderElectra/Source/HAPDecoderElectra/Private/HAPDecoder/ElectraMediaHAPDecoder.cpp b/Engine/Plugins/Media/HAPDecoderElectra/Source/HAPDecoderElectra/Private/HAPDecoder/ElectraMediaHAPDecoder.cpp
--- a/Engine/Plugins/Media/HAPDecoderElectra/Source/HAPDecoderElectra/Private/HAPDecoder/ElectraMediaHAPDecoder.cpp
+++ b/Engine/Plugins/Media/HAPDecoderElectra/Source/HAPDecoderElectra/Private/HAPDecoder/ElectraMediaHAPDecoder.cpp
@@ -665,7 +665,7 @@ IElectraDecoder::EDecoderError FVideoDecoderHAPElectra::DecodeAccessUnit(const F
 		{
 			auto& BufferInfo = NewOutput->Buffer[BufIdx];
 
-			uint8* BufferAddr;
+			uint8* BufferAddr = nullptr;
 			uint32 BufferPitch = 0;
 			uint32 ResourcePitch = 0;
 			uint32 BlockSizeX = InstanceBufferData[BufIdx].AllocBlockSize ? 4 : 1;
@@ -677,6 +677,8 @@ IElectraDecoder::EDecoderError FVideoDecoderHAPElectra::DecodeAccessUnit(const F
 			if (bUseGPUBuffers)
 			{
 				TRefCountPtr D3D12Device(static_cast<ID3D12Device*>(PlatformDevice));
+				HRESULT HResult;
+				FString HResultMsg;
 
 				// Create the resource pool as needed...
 				if (!InstanceBufferData[BufIdx].D3D12ResourcePool.IsValid())
@@ -691,23 +693,29 @@ IElectraDecoder::EDecoderError FVideoDecoderHAPElectra::DecodeAccessUnit(const F
 					TaskSync = FElectraDecodersPlatformResources::CreateAsyncConsecutiveTaskSync();
 				}
 
-				InstanceBufferData[BufIdx].D3D12ResourcePool->AllocateOutputDataAsBuffer(BufferInfo.GPUBuffer, ResourcePitch);
-
-				// Make the decoder write to a cached, temp buffer (otherwise performance is quite bad)
-				TempBuffer = MakeShared<TArray<uint8>, ESPMode::ThreadSafe>();
-				if (!TempBuffer.IsValid())
+				if (InstanceBufferData[BufIdx].D3D12ResourcePool->AllocateOutputDataAsBuffer(HResult, HResultMsg, BufferInfo.GPUBuffer, ResourcePitch))
 				{
-					PostError(Result, TEXT("Could not allocate HAP decoder output buffer"), ERRCODE_INTERNAL_FAILED_TO_DECODE_INPUT);
-					return IElectraDecoder::EDecoderError::Error;
+					// Make the decoder write to a cached, temp buffer (otherwise performance is quite bad)
+					TempBuffer = MakeShared<TArray<uint8>, ESPMode::ThreadSafe>();
+					if (!TempBuffer.IsValid())
+					{
+						PostError(Result, TEXT("Could not allocate HAP decoder output buffer"), ERRCODE_INTERNAL_FAILED_TO_DECODE_INPUT);
+						return IElectraDecoder::EDecoderError::Error;
+					}
+					TempBuffer->AddUninitialized(InstanceBufferData[BufIdx].AllocSize);
+					BufferAddr = TempBuffer->GetData();
+					if (BufferAddr == nullptr)
+					{
+						PostError(Result, TEXT("Could not allocate HAP decoder output buffer"), ERRCODE_INTERNAL_FAILED_TO_DECODE_INPUT);
+						return IElectraDecoder::EDecoderError::Error;
+					}
+					BufferPitch = (DecodedWidth / BlockSizeX) * BlockBytes;
 				}
-				TempBuffer->AddUninitialized(InstanceBufferData[BufIdx].AllocSize);
-				BufferAddr = TempBuffer->GetData();
-				if (BufferAddr == nullptr)
+				else
 				{
-					PostError(Result, TEXT("Could not allocate HAP decoder output buffer"), ERRCODE_INTERNAL_FAILED_TO_DECODE_INPUT);
+					PostError(Result, FString::Printf(TEXT("Could not allocate D3D12 buffer (%s)"), *HResultMsg), ERRCODE_INTERNAL_FAILED_TO_DECODE_INPUT);
 					return IElectraDecoder::EDecoderError::Error;
 				}
-				BufferPitch = (DecodedWidth / BlockSizeX) * BlockBytes;
 			}
 			else
 #endif
diff --git a/Engine/Plugins/Media/VPxDecoderElectra/Source/VPxDecoderElectra/Private/VPxDecoder/ElectraMediaVPxDecoder.cpp b/Engine/Plugins/Media/VPxDecoderElectra/Source/VPxDecoderElectra/Private/VPxDecoder/ElectraMediaVPxDecoder.cpp
--- a/Engine/Plugins/Media/VPxDecoderElectra/Source/VPxDecoderElectra/Private/VPxDecoder/ElectraMediaVPxDecoder.cpp
+++ b/Engine/Plugins/Media/VPxDecoderElectra/Source/VPxDecoderElectra/Private/VPxDecoder/ElectraMediaVPxDecoder.cpp
@@ -1020,9 +1020,9 @@ bool FVideoDecoderVPxElectra::ConvertDecodedImageToNV12orP010(FVideoDecoderOutpu
 		return false;
 	}
 
-	uint8* DstY;
-	uint8* DstUV;
-	uint32 DstPitch;
+	uint8* DstY = nullptr;
+	uint8* DstUV = nullptr;
+	uint32 DstPitch = 0;
 
 	TSharedPtr<TArray64<uint8>, ESPMode::ThreadSafe>& OutNV12Buffer = InNewOutput->ColorBuffer;
 #if ELECTRA_MEDIAGPUBUFFER_DX12
@@ -1047,6 +1047,8 @@ bool FVideoDecoderVPxElectra::ConvertDecodedImageToNV12orP010(FVideoDecoderOutpu
 	else
 	{
 		TRefCountPtr D3D12Device(static_cast<ID3D12Device*>(PlatformDevice));
+		HRESULT Result;
+		FString ResultMsg;
 
 		OutNV12Buffer.Reset();
 		InNewOutput->GPUBuffer.Resource = nullptr;
@@ -1059,16 +1061,24 @@ bool FVideoDecoderVPxElectra::ConvertDecodedImageToNV12orP010(FVideoDecoderOutpu
 
 		// Request resource and fence...
 		uint32 BufferPitch;
-		D3D12ResourcePool->AllocateOutputDataAsBuffer(InNewOutput->GPUBuffer, BufferPitch);
-
-		// Correct pitch to reflect resource's setup
-		InNewOutput->ColorPitch = BufferPitch;
+		if (D3D12ResourcePool->AllocateOutputDataAsBuffer(Result, ResultMsg, InNewOutput->GPUBuffer, BufferPitch))
+		{
+			// Correct pitch to reflect resource's setup
+			InNewOutput->ColorPitch = BufferPitch;
 
-		// Map the buffer so we can get a CPU address to the WC configured buffer
-		HRESULT Res = InNewOutput->GPUBuffer.Resource->Map(0, nullptr, (void**)&DstY);
-		check(SUCCEEDED(Res));
-		DstUV = DstY + BufferPitch * ah;
-		DstPitch = BufferPitch;
+			// Map the buffer so we can get a CPU address to the WC configured buffer
+			Result = InNewOutput->GPUBuffer.Resource->Map(0, nullptr, (void**)&DstY);
+			if (FAILED(Result))
+			{
+				return false;
+			}
+			DstUV = DstY + BufferPitch * ah;
+			DstPitch = BufferPitch;
+		}
+		else
+		{
+			return false;
+		}
 	}
 #endif // ELECTRA_MEDIAGPUBUFFER_DX12
 
@@ -1189,6 +1199,8 @@ bool FVideoDecoderVPxElectra::ConvertDecodedImageWithAlpha(FVideoDecoderOutputVP
 	else
 	{
 		TRefCountPtr D3D12Device(static_cast<ID3D12Device*>(PlatformDevice));
+		HRESULT Result;
+		FString ResultMsg;
 
 		OutPixelBuffer.Reset();
 		InNewOutput->GPUBuffer.Resource = nullptr;
@@ -1202,15 +1214,23 @@ bool FVideoDecoderVPxElectra::ConvertDecodedImageWithAlpha(FVideoDecoderOutputVP
 
 		// Request resource and fence...
 		uint32 BufferPitch;
-		D3D12ResourcePool->AllocateOutputDataAsBuffer(InNewOutput->GPUBuffer, BufferPitch);
-
-		// Map the buffer so we can get a CPU address to the WC configured buffer
-		HRESULT Res = InNewOutput->GPUBuffer.Resource->Map(0, nullptr, (void**)&Dst0);
-		check(SUCCEEDED(Res));
-		InNewOutput->ColorPitch = BufferPitch;
-		DstPitch = BufferPitch / sizeof(uint16);
-		LineSkipDst = ((DstPitch / 4) - aw) * 4;
-		Dst1 = Dst0 + DstPitch;
+		if (D3D12ResourcePool->AllocateOutputDataAsBuffer(Result, ResultMsg, InNewOutput->GPUBuffer, BufferPitch))
+		{
+			// Map the buffer so we can get a CPU address to the WC configured buffer
+			Result = InNewOutput->GPUBuffer.Resource->Map(0, nullptr, (void**)&Dst0);
+			if (FAILED(Result))
+			{
+				return false;
+			}
+			InNewOutput->ColorPitch = BufferPitch;
+			DstPitch = BufferPitch / sizeof(uint16);
+			LineSkipDst = ((DstPitch / 4) - aw) * 4;
+			Dst1 = Dst0 + DstPitch;
+		}
+		else
+		{
+			return false;
+		}
 	}
 #endif
 	InNewOutput->NumBits = 16;
