ElectraUtil: Removing hard limit of number of buffer slots in favor of dynamic resizes

#rnx

diff --git a/Engine/Plugins/Media/ElectraUtil/Source/ElectraSamples/Private/Windows/WindowsElectraTextureSampleGPUBufferHelper.cpp b/Engine/Plugins/Media/ElectraUtil/Source/ElectraSamples/Private/Windows/WindowsElectraTextureSampleGPUBufferHelper.cpp
--- a/Engine/Plugins/Media/ElectraUtil/Source/ElectraSamples/Private/Windows/WindowsElectraTextureSampleGPUBufferHelper.cpp
+++ b/Engine/Plugins/Media/ElectraUtil/Source/ElectraSamples/Private/Windows/WindowsElectraTextureSampleGPUBufferHelper.cpp
@@ -351,7 +351,7 @@ bool FElectraMediaDecoderOutputBufferPool_DX12::IsCompatibleAsTexture(uint32 InM
 // Check if a buffer is available
 bool FElectraMediaDecoderOutputBufferPool_DX12::BufferAvailable() const
 {
-	// If there is something in the pending free queue then we have room in som block.
+	// If there is something in the pending free queue then we have room in some block.
 	if (!PendingFreeQueue.IsEmpty())
 	{
 		return true;
@@ -364,9 +364,7 @@ bool FElectraMediaDecoderOutputBufferPool_DX12::BufferAvailable() const
 			return true;
 		}
 	}
-	// Can we still create a new entry?
-	const uint32 NumAllocatedSlots = ((Blocks.Num() - 1) * MaxBufsPerBlock) + Blocks.Last()->GetMaxNumBuffers();
-	return NumAllocatedSlots < MaxNumBuffers;
+	return false;
 }
 
 bool FElectraMediaDecoderOutputBufferPool_DX12::AllocateOutputDataAsBuffer(HRESULT& OutResult, FString& OutMessage, FOutputData& OutData, uint32& OutPitch)
@@ -388,11 +386,11 @@ bool FElectraMediaDecoderOutputBufferPool_DX12::AllocateOutputDataAsBuffer(HRESU
 		}
 	}
 
-	const uint32 NumAllocatedSlots = (Blocks.Num() - 1) * MaxBufsPerBlock + Blocks.Last()->GetMaxNumBuffers();
 	// Allocate a new block if we could not get a buffer, but only when there is no error yet.
-	if (!OutData.Resource.IsValid() && OutResult == S_OK && NumAllocatedSlots < MaxNumBuffers)
+	if (!OutData.Resource.IsValid() && OutResult == S_OK)
 	{
-		Blocks.Emplace(new FElectraMediaDecoderOutputBufferPoolBlock_DX12(Blocks.Num(), D3D12Device, MaxNumBuffers-NumAllocatedSlots > MaxBufsPerBlock ? MaxBufsPerBlock : MaxNumBuffers-NumAllocatedSlots, Width, Height, BytesPerPixel, D3D12HeapType));
+		const uint32 AdditionalBuffers = kMaxBufsPerBlock;
+		Blocks.Emplace(new FElectraMediaDecoderOutputBufferPoolBlock_DX12(Blocks.Num(), D3D12Device, AdditionalBuffers, Width, Height, BytesPerPixel, D3D12HeapType));
 		OutData.Resource = Blocks.Last()->AllocateOutputDataAsBuffer(OutResult, OutMessage, AsWeak(), OutPitch);
 	}
 
@@ -423,11 +421,11 @@ bool FElectraMediaDecoderOutputBufferPool_DX12::AllocateOutputDataAsTexture(HRES
 		}
 	}
 
-	const uint32 NumAllocatedSlots = (Blocks.Num() - 1) * MaxBufsPerBlock + Blocks.Last()->GetMaxNumBuffers();
 	// Allocate a new block if we could not get a buffer, but only when there is no error yet.
-	if (!OutData.Resource.IsValid() && OutResult == S_OK && NumAllocatedSlots < MaxNumBuffers)
+	if (!OutData.Resource.IsValid() && OutResult == S_OK)
 	{
-		Blocks.Emplace(new FElectraMediaDecoderOutputBufferPoolBlock_DX12(Blocks.Num(), D3D12Device, MaxNumBuffers-NumAllocatedSlots > MaxBufsPerBlock ? MaxBufsPerBlock : MaxNumBuffers-NumAllocatedSlots, Width, Height, PixFmt, D3D12HeapType));
+		const uint32 AdditionalBuffers = kMaxBufsPerBlock;
+		Blocks.Emplace(new FElectraMediaDecoderOutputBufferPoolBlock_DX12(Blocks.Num(), D3D12Device, AdditionalBuffers, Width, Height, PixFmt, D3D12HeapType));
 		OutData.Resource = Blocks.Last()->AllocateOutputDataAsTexture(OutResult, OutMessage, AsWeak(), InWidth, InHeight, InPixFmt);
 	}
 
