ElectraSamples: calling ShutdownPoolable() in the destructor to avoid potential resource leaks in case this has not been done by the owning pool, which may happen if it had been destroyed while samples are still actively used.

#rnx

diff --git a/Engine/Plugins/Media/ElectraUtil/Source/ElectraSamples/Private/Android/ElectraTextureSample_Android.cpp b/Engine/Plugins/Media/ElectraUtil/Source/ElectraSamples/Private/Android/ElectraTextureSample_Android.cpp
--- a/Engine/Plugins/Media/ElectraUtil/Source/ElectraSamples/Private/Android/ElectraTextureSample_Android.cpp
+++ b/Engine/Plugins/Media/ElectraUtil/Source/ElectraSamples/Private/Android/ElectraTextureSample_Android.cpp
@@ -745,6 +745,10 @@ bool FElectraTextureSample::IsReadyForReuse()
 
 void FElectraTextureSample::ShutdownPoolable()
 {
+	if (bWasShutDown)
+	{
+		return;
+	}
 //Q: this will trigger us to reuse the RHI/GL texture over and over -> sounds good. But is it 100%?
 //   - IF(!) the GL texture refs the EGL image internally we might keep the resources of that until a new one is assigned. Longer than needed! (although we free things up at the Java level (CleanupImageResources))
 // 	 - MIGHT that starve the decoder for buffers (we only have a max. number of "Images", too)
@@ -820,6 +824,8 @@ void FElectraTextureSample::CleanupImageResources()
 
 FElectraTextureSample::~FElectraTextureSample()
 {
+	ShutdownPoolable();
+
 	if (Buffer)
 	{
 		FMemory::Free(Buffer);
diff --git a/Engine/Plugins/Media/ElectraUtil/Source/ElectraSamples/Private/Apple/ElectraTextureSample_Apple.cpp b/Engine/Plugins/Media/ElectraUtil/Source/ElectraSamples/Private/Apple/ElectraTextureSample_Apple.cpp
--- a/Engine/Plugins/Media/ElectraUtil/Source/ElectraSamples/Private/Apple/ElectraTextureSample_Apple.cpp
+++ b/Engine/Plugins/Media/ElectraUtil/Source/ElectraSamples/Private/Apple/ElectraTextureSample_Apple.cpp
@@ -101,7 +101,7 @@ void FElectraMediaTexConvApple::ConvertTexture(FTextureRHIRef& InDstTexture, CVI
 
 			// Metal can upload directly from an IOSurface to a 2D texture
 			IMetalDynamicRHI* MetalDynamicRHI = GetIMetalDynamicRHI();
-				
+
 			TRefCountPtr<FRHITexture> YTex = MetalDynamicRHI->RHICreateTexture2DFromCVMetalTexture(bIs8Bit ? PF_G8 : PF_G16, ETextureCreateFlags::NoTiling | ETextureCreateFlags::ShaderResource, FClearValueBinding::White, YTextureRef);
 			TRefCountPtr<FRHITexture> UVTex = MetalDynamicRHI->RHICreateTexture2DFromCVMetalTexture(bIs8Bit ? PF_R8G8 : PF_G16R16, ETextureCreateFlags::NoTiling | ETextureCreateFlags::ShaderResource, FClearValueBinding::White, UVTextureRef);
 
@@ -140,13 +140,13 @@ void FElectraMediaTexConvApple::ConvertTexture(FTextureRHIRef& InDstTexture, CVI
 						SetGraphicsPipelineState(RHICmdList, GraphicsPSOInit, 0);
 
 						FShaderResourceViewRHIRef Y_SRV = RHICmdList.CreateShaderResourceView(
-							YTex, 
+							YTex,
 							FRHIViewDesc::CreateTextureSRV()
 								.SetDimensionFromTexture(YTex)
 								.SetMipRange(0, 1)
 								.SetFormat(PF_G8));
 						FShaderResourceViewRHIRef UV_SRV = RHICmdList.CreateShaderResourceView(
-							UVTex, 
+							UVTex,
 							FRHIViewDesc::CreateTextureSRV()
 								.SetDimensionFromTexture(UVTex)
 								.SetMipRange(0, 1)
@@ -167,18 +167,18 @@ void FElectraMediaTexConvApple::ConvertTexture(FTextureRHIRef& InDstTexture, CVI
 						SetGraphicsPipelineState(RHICmdList, GraphicsPSOInit, 0);
 
 						FShaderResourceViewRHIRef Y_SRV = RHICmdList.CreateShaderResourceView(
-							YTex, 
+							YTex,
 							FRHIViewDesc::CreateTextureSRV()
 								.SetDimensionFromTexture(YTex)
 								.SetMipRange(0, 1)
 								.SetFormat(PF_G16));
 						FShaderResourceViewRHIRef UV_SRV = RHICmdList.CreateShaderResourceView(
-							UVTex, 
+							UVTex,
 							FRHIViewDesc::CreateTextureSRV()
 								.SetDimensionFromTexture(UVTex)
 								.SetMipRange(0, 1)
 								.SetFormat(PF_G16R16));
-						
+
 						// Update shader uniform parameters.
 						SetShaderParametersLegacyPS(RHICmdList, PixelShader, FIntPoint(YWidth, YHeight), Y_SRV, UV_SRV, FIntPoint(YWidth, YHeight), YUVMtx, ColorSpaceMtx, EncodingType);
 					}
@@ -217,7 +217,7 @@ void FElectraMediaTexConvApple::ConvertTexture(FTextureRHIRef& InDstTexture, CVI
 				check(TextureRef);
 
 				IMetalDynamicRHI* MetalDynamicRHI = GetIMetalDynamicRHI();
-							
+
 				InDstTexture = MetalDynamicRHI->RHICreateTexture2DFromCVMetalTexture(PF_R16G16B16A16_UNORM, ETextureCreateFlags::NoTiling | ETextureCreateFlags::ShaderResource, FClearValueBinding::White, TextureRef);
 				CFRelease(TextureRef);
 			}
@@ -260,6 +260,7 @@ FElectraTextureSample::FElectraTextureSample(const TWeakPtr<FElectraMediaTexConv
 
 FElectraTextureSample::~FElectraTextureSample()
 {
+	ShutdownPoolable();
 }
 
 
@@ -270,6 +271,10 @@ FRHITexture* FElectraTextureSample::GetTexture() const
 
 void FElectraTextureSample::ShutdownPoolable()
 {
+	if (bWasShutDown)
+	{
+		return;
+	}
 	Buffer.Reset();
 	if (ImageBufferRef)
 	{
@@ -286,7 +291,7 @@ bool FElectraTextureSample::FinishInitialization()
         OSType PixelFormat = CVPixelBufferGetPixelFormatType(ImageBufferRef);
 		switch(PixelFormat)
 		{
-			case kCVPixelFormatType_4444AYpCbCr16: 
+			case kCVPixelFormatType_4444AYpCbCr16:
 			{
 				SampleFormat = EMediaTextureSampleFormat::Y416;
 				break;
@@ -346,7 +351,7 @@ bool FElectraTextureSample::FinishInitialization()
 				SampleFormat = EMediaTextureSampleFormat::BC4;
 				break;
 			}
-			case EPixelFormat::PF_NV12:  
+			case EPixelFormat::PF_NV12:
 			{
 				SampleFormat = EMediaTextureSampleFormat::CharNV12;
 				break;
diff --git a/Engine/Plugins/Media/ElectraUtil/Source/ElectraSamples/Private/ElectraTextureSampleBase.cpp b/Engine/Plugins/Media/ElectraUtil/Source/ElectraSamples/Private/ElectraTextureSampleBase.cpp
--- a/Engine/Plugins/Media/ElectraUtil/Source/ElectraSamples/Private/ElectraTextureSampleBase.cpp
+++ b/Engine/Plugins/Media/ElectraUtil/Source/ElectraSamples/Private/ElectraTextureSampleBase.cpp
@@ -255,6 +255,8 @@ void IElectraTextureSampleBase::InitializeCommon(const Electra::FParamDict& InPa
 	CropBottom = (int32)InParams.GetValue(IDecoderOutputOptionNames::CropBottom).SafeGetInt64(0);
 	ImageDim.X = ImageOutputDim.X + (CropLeft + CropRight);
 	ImageDim.Y = ImageOutputDim.Y + (CropTop + CropBottom);
+	// Newly initialized now.
+	bWasShutDown = false;
 }
 
 bool IElectraTextureSampleBase::FinishInitialization()
@@ -270,6 +272,11 @@ IElectraTextureSampleBase::FReleaseDelegate& IElectraTextureSampleBase::GetRelea
 
 // -------------------------------------------------------------------------------------------------------------------------------------------------------
 
+IElectraTextureSampleBase::~IElectraTextureSampleBase()
+{
+	ShutdownPoolable();
+}
+
 bool IElectraTextureSampleBase::IsCacheable() const
 {
 	return true;
@@ -281,9 +288,13 @@ void IElectraTextureSampleBase::InitializePoolable()
 
 void IElectraTextureSampleBase::ShutdownPoolable()
 {
-	ReleaseDelegate.ExecuteIfBound(this);
-	// Clean out what is needed for this element to be reused.
-	ReleaseDelegate.Unbind();
+	if (!bWasShutDown)
+	{
+		ReleaseDelegate.ExecuteIfBound(this);
+		// Clean out what is needed for this element to be reused.
+		ReleaseDelegate.Unbind();
+		bWasShutDown = true;
+	}
 }
 
 bool IElectraTextureSampleBase::IsReadyForReuse()
diff --git a/Engine/Plugins/Media/ElectraUtil/Source/ElectraSamples/Private/Linux/ElectraTextureSample_Linux.cpp b/Engine/Plugins/Media/ElectraUtil/Source/ElectraSamples/Private/Linux/ElectraTextureSample_Linux.cpp
--- a/Engine/Plugins/Media/ElectraUtil/Source/ElectraSamples/Private/Linux/ElectraTextureSample_Linux.cpp
+++ b/Engine/Plugins/Media/ElectraUtil/Source/ElectraSamples/Private/Linux/ElectraTextureSample_Linux.cpp
@@ -6,6 +6,7 @@
 
 FElectraTextureSample::~FElectraTextureSample()
 {
+	ShutdownPoolable();
 }
 
 uint32 FElectraTextureSample::GetStride() const
@@ -36,6 +37,10 @@ IMediaTextureSampleConverter* FElectraTextureSample::GetMediaTextureSampleConver
 #if !UE_SERVER
 void FElectraTextureSample::ShutdownPoolable()
 {
+	if (bWasShutDown)
+	{
+		return;
+	}
 	Buffer.Reset();
 	IElectraTextureSampleBase::ShutdownPoolable();
 }
diff --git a/Engine/Plugins/Media/ElectraUtil/Source/ElectraSamples/Private/Windows/ElectraTextureSample_Windows.cpp b/Engine/Plugins/Media/ElectraUtil/Source/ElectraSamples/Private/Windows/ElectraTextureSample_Windows.cpp
--- a/Engine/Plugins/Media/ElectraUtil/Source/ElectraSamples/Private/Windows/ElectraTextureSample_Windows.cpp
+++ b/Engine/Plugins/Media/ElectraUtil/Source/ElectraSamples/Private/Windows/ElectraTextureSample_Windows.cpp
@@ -45,6 +45,7 @@ DECLARE_GPU_STAT_NAMED(MediaWinDecoder_Convert, TEXT("MediaWinDecoder_Convert"))
 
 FElectraTextureSample::~FElectraTextureSample()
 {
+	ShutdownPoolable();
 	// Release DX11 resources asynchronously. There is some internal synchronization going on that takes a while
 	// to finish that would block the game thread.
 	if (GIsRunning && (TextureDX11.IsValid() || D3D11Device.IsValid()))
@@ -295,6 +296,10 @@ bool FElectraTextureSample::FinishInitialization()
 #if !UE_SERVER
 void FElectraTextureSample::ShutdownPoolable()
 {
+	if (bWasShutDown)
+	{
+		return;
+	}
 	// If we use DX12 with texture data, we use the RHI texture only to refer to the actual "native" D3D resource we created -> get rid of this "alias" as soon as we retire to the pool
 	if (TextureDX12 && RHIGetInterfaceType() == ERHIInterfaceType::D3D12)
 	{
diff --git a/Engine/Plugins/Media/ElectraUtil/Source/ElectraSamples/Public/IElectraTextureSample.h b/Engine/Plugins/Media/ElectraUtil/Source/ElectraSamples/Public/IElectraTextureSample.h
--- a/Engine/Plugins/Media/ElectraUtil/Source/ElectraSamples/Public/IElectraTextureSample.h
+++ b/Engine/Plugins/Media/ElectraUtil/Source/ElectraSamples/Public/IElectraTextureSample.h
@@ -23,6 +23,8 @@ class IElectraTextureSampleBase
 	, public IMediaPoolable
 {
 public:
+	UE_API virtual ~IElectraTextureSampleBase();
+
 	UE_API void InitializeCommon(const Electra::FParamDict& InParams);
 	UE_API virtual bool FinishInitialization();
 
@@ -113,6 +115,7 @@ protected:
 
 	/** Optional delegate to call during ShutdownPoolable(). */
 	FReleaseDelegate ReleaseDelegate;
+	bool bWasShutDown = false;
 };
 
 #undef UE_API
