ElectraSamples: Fixed Linux server build linker error

#rnx

diff --git a/Engine/Plugins/Media/ElectraUtil/Source/ElectraSamples/Private/ElectraTextureSampleBase.cpp b/Engine/Plugins/Media/ElectraUtil/Source/ElectraSamples/Private/ElectraTextureSampleBase.cpp
--- a/Engine/Plugins/Media/ElectraUtil/Source/ElectraSamples/Private/ElectraTextureSampleBase.cpp
+++ b/Engine/Plugins/Media/ElectraUtil/Source/ElectraSamples/Private/ElectraTextureSampleBase.cpp
@@ -463,6 +463,9 @@ float IElectraTextureSampleBase::GetSampleDataScale(bool b10Bit) const
 
 #else
 
+IElectraTextureSampleBase::~IElectraTextureSampleBase()
+{ }
+
 void IElectraTextureSampleBase::InitializeCommon(const Electra::FParamDict& InParams)
 { }
 
diff --git a/Engine/Plugins/Media/ElectraUtil/Source/ElectraSamples/Private/Linux/ElectraTextureSample_Linux.cpp b/Engine/Plugins/Media/ElectraUtil/Source/ElectraSamples/Private/Linux/ElectraTextureSample_Linux.cpp
--- a/Engine/Plugins/Media/ElectraUtil/Source/ElectraSamples/Private/Linux/ElectraTextureSample_Linux.cpp
+++ b/Engine/Plugins/Media/ElectraUtil/Source/ElectraSamples/Private/Linux/ElectraTextureSample_Linux.cpp
@@ -6,7 +6,9 @@
 
 FElectraTextureSample::~FElectraTextureSample()
 {
+#if !UE_SERVER
 	ShutdownPoolable();
+#endif
 }
 
 uint32 FElectraTextureSample::GetStride() const
