ElectraUtils: Fixed check to BufferAvailable() in the DX12 buffer helpers

#rnx

diff --git a/Engine/Plugins/Media/ElectraUtil/Source/ElectraSamples/Private/Windows/WindowsElectraTextureSampleGPUBufferHelper.cpp b/Engine/Plugins/Media/ElectraUtil/Source/ElectraSamples/Private/Windows/WindowsElectraTextureSampleGPUBufferHelper.cpp
--- a/Engine/Plugins/Media/ElectraUtil/Source/ElectraSamples/Private/Windows/WindowsElectraTextureSampleGPUBufferHelper.cpp
+++ b/Engine/Plugins/Media/ElectraUtil/Source/ElectraSamples/Private/Windows/WindowsElectraTextureSampleGPUBufferHelper.cpp
@@ -357,14 +357,17 @@ bool FElectraMediaDecoderOutputBufferPool_DX12::BufferAvailable() const
 		return true;
 	}
 	// Check if any already existing block has an available entry
+	uint32 NumAllocatedBuffers = 0;
 	for(FElectraMediaDecoderOutputBufferPoolBlock_DX12* Block : Blocks)
 	{
+		NumAllocatedBuffers += Block->MaxNumBuffers;
 		if (Block->BufferAvailable())
 		{
 			return true;
 		}
 	}
-	return false;
+	// Still possible to add new buffers?
+	return NumAllocatedBuffers < MaxNumBuffers;
 }
 
 bool FElectraMediaDecoderOutputBufferPool_DX12::AllocateOutputDataAsBuffer(HRESULT& OutResult, FString& OutMessage, FOutputData& OutData, uint32& OutPitch)
