Fixed some Bughawk JNI issues

#android

diff --git a/Engine/Plugins/Media/AndroidCamera/Source/AndroidCamera/Private/AndroidJavaCameraPlayer.cpp b/Engine/Plugins/Media/AndroidCamera/Source/AndroidCamera/Private/AndroidJavaCameraPlayer.cpp
--- a/Engine/Plugins/Media/AndroidCamera/Source/AndroidCamera/Private/AndroidJavaCameraPlayer.cpp
+++ b/Engine/Plugins/Media/AndroidCamera/Source/AndroidCamera/Private/AndroidJavaCameraPlayer.cpp
@@ -435,7 +435,8 @@ jobject FJavaAndroidCameraPlayer::GetVideoLastMediaFrameData(int32& CurrentPosit
 
 			UpdateScaleRotationOffset();
 
-			jobject GlobalBRFrameData = FAndroidApplication::GetJavaEnv()->NewGlobalRef(BRFrameData);
+			jobject GlobalBRFrameData = JEnv->NewGlobalRef(BRFrameData);
+			JEnv->DeleteLocalRef(BRFrameData);
 			return GlobalBRFrameData;
 		}
 	}
diff --git a/Engine/Plugins/Media/ElectraCodecs/Source/ElectraDecoders/Private/Android/h264/VideoDecoderH264_JavaWrapper_Android.cpp b/Engine/Plugins/Media/ElectraCodecs/Source/ElectraDecoders/Private/Android/h264/VideoDecoderH264_JavaWrapper_Android.cpp
--- a/Engine/Plugins/Media/ElectraCodecs/Source/ElectraDecoders/Private/Android/h264/VideoDecoderH264_JavaWrapper_Android.cpp
+++ b/Engine/Plugins/Media/ElectraCodecs/Source/ElectraDecoders/Private/Android/h264/VideoDecoderH264_JavaWrapper_Android.cpp
@@ -902,7 +902,8 @@ int32 FElectraH264VideoDecoderAndroidJava::GetOutputBuffer(void*& OutBufferDataP
 /**
  * Releases the decoder output buffer back to the decoder.
  *
- * @param InOutBufferInfo
+ * @param BufferIndex
+ * @param ValidCount
  * @param bRender
  * @param releaseAt
  *
@@ -917,7 +918,7 @@ int32 FElectraH264VideoDecoderAndroidJava::ReleaseOutputBuffer(int32 BufferIndex
 		if (ValidCount == CurrentValidCount)
 		{
 			// Yes...
-			int32 result = CallMethodNoVerify<int>(ReleaseOutputBufferFN, BufferIndex, bRender, (long)releaseAt);
+			int32 result = CallMethodNoVerify<int>(ReleaseOutputBufferFN, BufferIndex, bRender, releaseAt);
 			GClearException();
 			return result ? 1 : 0;
 		}
diff --git a/Engine/Plugins/Media/ElectraCodecs/Source/ElectraDecoders/Private/Android/h265/VideoDecoderH265_JavaWrapper_Android.cpp b/Engine/Plugins/Media/ElectraCodecs/Source/ElectraDecoders/Private/Android/h265/VideoDecoderH265_JavaWrapper_Android.cpp
--- a/Engine/Plugins/Media/ElectraCodecs/Source/ElectraDecoders/Private/Android/h265/VideoDecoderH265_JavaWrapper_Android.cpp
+++ b/Engine/Plugins/Media/ElectraCodecs/Source/ElectraDecoders/Private/Android/h265/VideoDecoderH265_JavaWrapper_Android.cpp
@@ -893,7 +893,8 @@ int32 FElectraH265VideoDecoderAndroidJava::GetOutputBuffer(void*& OutBufferDataP
 /**
  * Releases the decoder output buffer back to the decoder.
  *
- * @param InOutBufferInfo
+ * @param BufferIndex
+ * @param ValidCount
  * @param bRender
  * @param releaseAt
  *
@@ -908,7 +909,7 @@ int32 FElectraH265VideoDecoderAndroidJava::ReleaseOutputBuffer(int32 BufferIndex
 		if (ValidCount == CurrentValidCount)
 		{
 			// Yes...
-			int32 result = CallMethodNoVerify<int>(ReleaseOutputBufferFN, BufferIndex, bRender, (long)releaseAt);
+			int32 result = CallMethodNoVerify<int>(ReleaseOutputBufferFN, BufferIndex, bRender, releaseAt);
 			GClearException();
 			return result ? 1 : 0;
 		}
diff --git a/Engine/Plugins/Media/ElectraCodecs/Source/ElectraDecoders/Private/Android/vpx/VideoDecoderVPx_JavaWrapper_Android.cpp b/Engine/Plugins/Media/ElectraCodecs/Source/ElectraDecoders/Private/Android/vpx/VideoDecoderVPx_JavaWrapper_Android.cpp
--- a/Engine/Plugins/Media/ElectraCodecs/Source/ElectraDecoders/Private/Android/vpx/VideoDecoderVPx_JavaWrapper_Android.cpp
+++ b/Engine/Plugins/Media/ElectraCodecs/Source/ElectraDecoders/Private/Android/vpx/VideoDecoderVPx_JavaWrapper_Android.cpp
@@ -891,7 +891,8 @@ int32 FElectraVPxVideoDecoderAndroidJava::GetOutputBuffer(void*& OutBufferDataPt
 /**
  * Releases the decoder output buffer back to the decoder.
  *
- * @param InOutBufferInfo
+ * @param BufferIndex
+ * @param ValidCount
  * @param bRender
  * @param releaseAt
  *
@@ -906,7 +907,7 @@ int32 FElectraVPxVideoDecoderAndroidJava::ReleaseOutputBuffer(int32 BufferIndex,
 		if (ValidCount == CurrentValidCount)
 		{
 			// Yes...
-			int32 result = CallMethodNoVerify<int>(ReleaseOutputBufferFN, BufferIndex, bRender, (long)releaseAt);
+			int32 result = CallMethodNoVerify<int>(ReleaseOutputBufferFN, BufferIndex, bRender, releaseAt);
 			GClearException();
 			return result ? 1 : 0;
 		}
diff --git a/Engine/Plugins/Media/ElectraUtil/Source/ElectraSamples/Private/Android/ElectraTextureSample_Android.cpp b/Engine/Plugins/Media/ElectraUtil/Source/ElectraSamples/Private/Android/ElectraTextureSample_Android.cpp
--- a/Engine/Plugins/Media/ElectraUtil/Source/ElectraSamples/Private/Android/ElectraTextureSample_Android.cpp
+++ b/Engine/Plugins/Media/ElectraUtil/Source/ElectraSamples/Private/Android/ElectraTextureSample_Android.cpp
@@ -274,7 +274,7 @@ FElectraTextureSampleSupport::FElectraTextureSampleSupport()
 	FImageResources_VScale = FindField(JEnv, FImageResourcesClass, "VScale", "F", false);
 	FImageResources_VOffset = FindField(JEnv, FImageResourcesClass, "VOffset", "F", false);
 
-	FImageResources_ReleaseFN = JEnv->GetMethodID(FImageResourcesClass, "Release", "()V");
+	FImageResources_ReleaseFN = FindMethod(JEnv, FImageResourcesClass, "Release", "()V", false);
 
 	// - - - - - - - - - - - - - - - - - - - - - - - - - - - - - - - - - - - - - - - - - - - - - - - - - - - - - - - - - - - - - - -
 
@@ -287,8 +287,16 @@ FElectraTextureSampleSupport::FElectraTextureSampleSupport()
 		{
 			SDKint = JEnv->GetStaticIntField(localVersionClass, SdkIntFieldID);
 		}
+		else
+		{
+			JEnv->ExceptionClear();
+		}
 		JEnv->DeleteLocalRef(localVersionClass);
 	}
+	else
+	{
+		JEnv->ExceptionClear();
+	}
 
 	// - - - - - - - - - - - - - - - - - - - - - - - - - - - - - - - - - - - - - - - - - - - - - - - - - - - - - - - - - - - - - - -
 
@@ -326,7 +334,7 @@ FElectraTextureSampleSupport::FElectraTextureSampleSupport()
 	FlushRenderingCommands();
 #else
 	// Setup Java side of things
-	JEnv->CallVoidMethod(Object, InitializeFN.Method, FAndroidMisc::ShouldUseVulkan());
+	JEnv->CallVoidMethod(Object, InitializeFN.Method, UseGpuOutputPath(), FAndroidMisc::ShouldUseVulkan(), (jlong)this);
 
 	// Query surface to be used for decoder
 	jobject Surface = JEnv->CallObjectMethod(Object, GetCodecSurfaceFN.Method);
@@ -375,9 +383,11 @@ FElectraTextureSampleSupport::~FElectraTextureSampleSupport()
 				{
 					JEnv->DeleteGlobalRef(CodecSurfaceToDelete);
 				}
+				JEnv->DeleteGlobalRef(FImageResourcesClass);
 				JEnv->DeleteGlobalRef(FFrameUpdateInfoClass);
 				CodecSurfaceToDelete = nullptr;
-				FFrameUpdateInfoClass = 0;
+				FImageResourcesClass = nullptr;
+				FFrameUpdateInfoClass = nullptr;
 			});
 		});
 		FlushRenderingCommands();
@@ -390,9 +400,11 @@ FElectraTextureSampleSupport::~FElectraTextureSampleSupport()
 		{
 			JEnv->DeleteGlobalRef(CodecSurfaceToDelete);
 		}
+		JEnv->DeleteGlobalRef(FImageResourcesClass);
 		JEnv->DeleteGlobalRef(FFrameUpdateInfoClass);
 		CodecSurfaceToDelete = nullptr;
-		FFrameUpdateInfoClass = 0;
+		FImageResourcesClass = nullptr;
+		FFrameUpdateInfoClass = nullptr;
 	}
 }
 
diff --git a/Engine/Plugins/Online/Android/AndroidFetchBackgroundDownload/Source/AndroidFetchBackgroundDownload/Private/AndroidPlatformBackgroundHttpManager.cpp b/Engine/Plugins/Online/Android/AndroidFetchBackgroundDownload/Source/AndroidFetchBackgroundDownload/Private/AndroidPlatformBackgroundHttpManager.cpp
--- a/Engine/Plugins/Online/Android/AndroidFetchBackgroundDownload/Source/AndroidFetchBackgroundDownload/Private/AndroidPlatformBackgroundHttpManager.cpp
+++ b/Engine/Plugins/Online/Android/AndroidFetchBackgroundDownload/Source/AndroidFetchBackgroundDownload/Private/AndroidPlatformBackgroundHttpManager.cpp
@@ -907,6 +907,7 @@ void FAndroidPlatformBackgroundHttpManager::Java_OnTick(JNIEnv* Env, jobject Und
 				if (bIsThreadsafeToHandle)
 				{
 					jmethodID CancelMethod = Env->GetMethodID(UnderlyingWorkerClass, "CancelRequest", "(Ljava/lang/String;)V");
+					CHECK_JNI_RESULT(CancelMethod);
 
 					for (const FString& RequestID : RequestsToCancelByID_Java)
 					{
