Electra: Added support for fragmented mp4 files in Protron

#rnx

#Synced-CL Unknown

diff --git a/Engine/Plugins/Media/ElectraPlayer/Source/ElectraPlayerRuntime/Private/Runtime/StreamTypes.cpp b/Engine/Plugins/Media/ElectraPlayer/Source/ElectraPlayerRuntime/Private/Runtime/StreamTypes.cpp
--- a/Engine/Plugins/Media/ElectraPlayer/Source/ElectraPlayerRuntime/Private/Runtime/StreamTypes.cpp
+++ b/Engine/Plugins/Media/ElectraPlayer/Source/ElectraPlayerRuntime/Private/Runtime/StreamTypes.cpp
@@ -606,6 +606,12 @@ namespace Electra
 						HumanReadableCodecName = TEXT("Avid DNxHD");
 						return true;
 					}
+
+					case ElectraDecodersUtil::Make4CC('a','p','v','1'):
+					{
+						HumanReadableCodecName = TEXT("Advanced Professional Video (APV)");
+						return true;
+					}
 				}
 				HumanReadableCodecName = Printable4CC(GetCodec4CC());
 				return true;
diff --git a/Engine/Plugins/Media/ElectraPlayer/Source/ElectraProtron/Private/Player/ElectraProtronPlayerImpl.cpp b/Engine/Plugins/Media/ElectraPlayer/Source/ElectraProtron/Private/Player/ElectraProtronPlayerImpl.cpp
--- a/Engine/Plugins/Media/ElectraPlayer/Source/ElectraProtron/Private/Player/ElectraProtronPlayerImpl.cpp
+++ b/Engine/Plugins/Media/ElectraPlayer/Source/ElectraProtron/Private/Player/ElectraProtronPlayerImpl.cpp
@@ -244,21 +244,24 @@ void FElectraProtronPlayer::FImpl::InternalOpen(const FString& InFilename)
 	ParsedRootBoxes.SetNum(RootBoxes.Num());
 	for(int32 i=0; i<RootBoxes.Num(); ++i)
 	{
-		if (!ParsedRootBoxes[i].ParseBoxTree(RootBoxes[i]))
+		Electra::UtilitiesMP4::FMP4BoxTreeParser tp;
+		if (!tp.ParseBoxTree(RootBoxes[i]))
 		{
 			LastErrorMessage = FString::Printf(TEXT("Failed to parse the file's box structure."));
 			return;
 		}
+		ParsedRootBoxes[i] = tp.GetBoxTree();
 	}
+	auto FindRootBox = [&](uint32 InType) -> TSharedPtr<Electra::UtilitiesMP4::FMP4BoxBase>
+	{
+		auto rb = ParsedRootBoxes.FindByPredicate([Type=InType](const TSharedPtr<Electra::UtilitiesMP4::FMP4BoxBase>& In) { return In->GetType() == Type;});
+		return rb ? *rb : nullptr;
+	};
 
 	// We need the `moov` box now to determine what is inside this file.
 	// Yes, we looked at the presence of it above already as a quick reject before parsing the boxes.
 	// Now we need it for real and we know that it needs to be there, so here we go
-	TSharedPtr<Electra::UtilitiesMP4::FMP4BoxMOOV, ESPMode::ThreadSafe> MoovBox = StaticCastSharedPtr<Electra::UtilitiesMP4::FMP4BoxMOOV>(
-															(ParsedRootBoxes.FindByPredicate([](const Electra::UtilitiesMP4::FMP4BoxTreeParser& In)
-															{
-																return In.GetBoxTree()->GetType() == Electra::UtilitiesMP4::MakeBoxAtom('m','o','o','v');
-															})->GetBoxTree()));
+	TSharedPtr<Electra::UtilitiesMP4::FMP4BoxMOOV, ESPMode::ThreadSafe> MoovBox = StaticCastSharedPtr<Electra::UtilitiesMP4::FMP4BoxMOOV>(FindRootBox(Electra::UtilitiesMP4::MakeBoxAtom('m','o','o','v')));
 	// Get the `mvhd` box for the duration of the movie and the timescale other values are measured in.
 	auto MvhdBox = MoovBox->FindBoxRecursive<Electra::UtilitiesMP4::FMP4BoxMVHD>(Electra::UtilitiesMP4::MakeBoxAtom('m','v','h','d'), 0);
 	if (!MvhdBox.IsValid())
@@ -266,15 +269,44 @@ void FElectraProtronPlayer::FImpl::InternalOpen(const FString& InFilename)
 		LastErrorMessage = FString::Printf(TEXT("No `mvhd` box found. This is not a usable mp4/mov file."));
 		return;
 	}
+
+	// Check if this is a fragmented file. We do not look for `sidx` or `mfra` boxes since these are purely optional.
+	TSharedPtr<Electra::UtilitiesMP4::FMP4BoxMOOF, ESPMode::ThreadSafe> FirstMoofBox = StaticCastSharedPtr<Electra::UtilitiesMP4::FMP4BoxMOOF>(FindRootBox(Electra::UtilitiesMP4::MakeBoxAtom('m','o','o','f')));
+	auto MvexBox = MoovBox->FindBoxRecursive<Electra::UtilitiesMP4::FMP4BoxMVEX>(Electra::UtilitiesMP4::MakeBoxAtom('m','v','e','x'), 0);
+	if (FirstMoofBox && !MvexBox)
+	{
+		LastErrorMessage = FString::Printf(TEXT("No `mvex` box present in a file containing `moof` fragment boxes."));
+		return;
+	}
+	const bool bIsFragmented = MvexBox || FirstMoofBox;
+
 	MovieDuration = MvhdBox->GetDuration();
 	if (!MovieDuration.IsValid())
 	{
 		LastErrorMessage = FString::Printf(TEXT("Duration in `mvhd` box is set to indefinite or is not valid. This is not a usable mp4/mov file."));
 		return;
 	}
+	// If there is a movie extends header box we take the movie duration from there.
+	if (bIsFragmented && MvexBox)
+	{
+		auto MehdBox = MvexBox->FindBoxRecursive<Electra::UtilitiesMP4::FMP4BoxMEHD>(Electra::UtilitiesMP4::MakeBoxAtom('m','e','h','d'), 0);
+		/*
+			From ISO/IEC 14496-12:2022 section 8.8.2.1
+				If movie fragments are present but there is no MediaExtendsHeaderBox
+				and the movie duration is 0, the movie duration should be interpreted as indefinite duration.
+		*/
+		if (MehdBox)
+		{
+			MovieDuration.SetNumerator(MehdBox->GetFragmentDuration());
+		}
+		else if (MovieDuration.GetNumerator() == 0)
+		{
+			LastErrorMessage = FString::Printf(TEXT("Duration in `mvhd` box is set to zero in a fragmented file that is missing a `mehd` box. This is not a usable mp4/mov file."));
+			return;
+		}
+	}
 	Duration = MovieDuration.GetAsTimebase(ETimespan::TicksPerSecond);
 
-
 	IElectraCodecFactoryModule* FactoryModule = static_cast<IElectraCodecFactoryModule*>(FModuleManager::Get().GetModule(TEXT("ElectraCodecFactory")));
 	check(FactoryModule);
 	if (!FactoryModule)
@@ -294,17 +326,25 @@ void FElectraProtronPlayer::FImpl::InternalOpen(const FString& InFilename)
 	int64 ShortestSupportedTrackDuration = TNumericLimits<int64>::Max();
 	for(int32 i=0; i<AllTracks.Num(); ++i)
 	{
+		const TSharedPtr<Electra::UtilitiesMP4::FMP4BoxTRAK, ESPMode::ThreadSafe>& ThisTrakBox(AllTracks[i]);
+
 		Tracks[i] = MakeShared<FTrackInfo>();
-		Tracks[i]->TrackBox = AllTracks[i];
+		Tracks[i]->TrackBox = ThisTrakBox;
 
-		auto Tkhd = AllTracks[i]->FindBoxRecursive<Electra::UtilitiesMP4::FMP4BoxTKHD>(Electra::UtilitiesMP4::MakeBoxAtom('t','k','h','d'), 1);
+		auto Tkhd = ThisTrakBox->FindBoxRecursive<Electra::UtilitiesMP4::FMP4BoxTKHD>(Electra::UtilitiesMP4::MakeBoxAtom('t','k','h','d'), 1);
+		auto Mdhd = ThisTrakBox->FindBoxRecursive<Electra::UtilitiesMP4::FMP4BoxMDHD>(Electra::UtilitiesMP4::MakeBoxAtom('m','d','h','d'), 3);
 		if (!Tkhd)
 		{
-			// If not found the file is broken.
 			LastErrorMessage = FString::Printf(TEXT("No `tkhd` box found on track %d. This file cannot be used."), i);
 			Tracks.Empty();
 			return;
 		}
+		if (!Mdhd)
+		{
+			LastErrorMessage = FString::Printf(TEXT("No `mdhd` box found on track %d. This file cannot be used."), i);
+			Tracks.Empty();
+			return;
+		}
 		uint32 TrackID = Tracks[i]->TrackID = Tkhd->GetTrackID();
 		/*
 			From the standard:
@@ -318,17 +358,39 @@ void FElectraProtronPlayer::FImpl::InternalOpen(const FString& InFilename)
 		}
 
 		// This track's duration must not be indefinite.
-		if (Tkhd->GetDuration() == TNumericLimits<int64>::Max())
+		int64 TrackDuration = Tkhd->GetDuration();
+		if (TrackDuration == TNumericLimits<int64>::Max())
 		{
 			UE_LOG(LogElectraProtron, Warning, TEXT("Track #%u has indefinite duration, ignoring this track."), TrackID);
 			continue;
 		}
 
+		// If the file is fragmented we collect all the track fragment boxes for this track.
+		if (bIsFragmented)
+		{
+			Tracks[i]->FragmentInfo = MakeShared<Electra::UtilitiesMP4::FMP4Track::FFragmentInfo, ESPMode::ThreadSafe>();
+			if (!Tracks[i]->FragmentInfo->Prepare(ThisTrakBox, MvexBox, ParsedRootBoxes))
+			{
+				UE_LOG(LogElectraProtron, Warning, TEXT("Fragmented track #%u could not be prepared, ignoring this track."), TrackID);
+				continue;
+			}
+
+			// The summed duration is in media timescale, so we need to transform it into movie timescale.
+			Electra::FTimeFraction MediaDur(Tracks[i]->FragmentInfo->GetTotalDuration(), Mdhd->GetTimescale());
+			TrackDuration = MediaDur.GetAsTimebase(MovieDuration.GetDenominator());
+			// It is possible that the movie duration is only describing the content contained in the moov box.
+			// If that is smaller than the track duration, update the movie duration with the track duration.
+			if (MovieDuration.GetNumerator() < TrackDuration)
+			{
+				MovieDuration.SetNumerator(TrackDuration);
+			}
+		}
+
 		// Take note of the track with the longest duration (any track, even unsupported ones)
-		LongestTrackDuration = LongestTrackDuration < Tkhd->GetDuration() ? Tkhd->GetDuration() : LongestTrackDuration;
+		LongestTrackDuration = LongestTrackDuration < TrackDuration ? TrackDuration : LongestTrackDuration;
 
 		ElectraProtronUtils::FCodecInfo& CodecInfo(Tracks[i]->CodecInfo);
-		GetTrackCodecInfo(CodecInfo, AllTracks[i], TrackID);
+		GetTrackCodecInfo(CodecInfo, ThisTrakBox, TrackID);
 
 		// Timecode tracks are inherently usable.
 		if (CodecInfo.Type == ElectraProtronUtils::FCodecInfo::EType::Timecode)
@@ -405,9 +467,9 @@ void FElectraProtronPlayer::FImpl::InternalOpen(const FString& InFilename)
 				Tracks[i]->DecoderProviderInfo.Implementation = Factory->GetProviderInformation().GetImplementation();
 				Tracks[i]->DecoderProviderInfo.Vendor = Factory->GetProviderInformation().GetVendor();
 				// Take note of the supported track with the longest duration.
-				LongestSupportedTrackDuration = LongestSupportedTrackDuration < Tkhd->GetDuration() ? Tkhd->GetDuration() : LongestSupportedTrackDuration;
+				LongestSupportedTrackDuration = LongestSupportedTrackDuration < TrackDuration ? TrackDuration : LongestSupportedTrackDuration;
 				// Likewise for the shortest.
-				ShortestSupportedTrackDuration = ShortestSupportedTrackDuration > Tkhd->GetDuration() ? Tkhd->GetDuration() : ShortestSupportedTrackDuration;
+				ShortestSupportedTrackDuration = ShortestSupportedTrackDuration > TrackDuration ? TrackDuration : ShortestSupportedTrackDuration;
 			}
 			else
 			{
@@ -428,23 +490,13 @@ void FElectraProtronPlayer::FImpl::InternalOpen(const FString& InFilename)
 	}
 	Electra::FTimeFraction EntireMovieDuration(MovieDuration);
 
-#if 0
-	// Check that the movie duration is not larger than the longest supported track.
-	if (LongestSupportedTrackDuration > 0 && MvhdDur > LongestSupportedTrackDuration)
-	{
-		UE_LOG(LogElectraProtron, Warning, TEXT("Movie duration in `mvhd` box (%lld) is larger than the longest supported track (%lld) in the file, adjusting movie duration down."), (long long int)MvhdDur, (long long int)LongestSupportedTrackDuration);
-		MovieDuration.SetFromND(LongestSupportedTrackDuration, MovieDuration.GetDenominator());
-		Duration = MovieDuration.GetAsTimebase(ETimespan::TicksPerSecond);
-	}
-#else
-	// Check that the movie duration is not larger than the longest supported track.
+	// Check that the movie duration is not larger than the shortest supported track.
 	if (ShortestSupportedTrackDuration < TNumericLimits<int64>::Max() && MvhdDur > ShortestSupportedTrackDuration)
 	{
 		UE_LOG(LogElectraProtron, Warning, TEXT("Movie duration in `mvhd` box (%lld) is larger than the shortest supported track (%lld) in the file, adjusting movie duration down."), (long long int)MvhdDur, (long long int)ShortestSupportedTrackDuration);
 		MovieDuration.SetFromND(ShortestSupportedTrackDuration, MovieDuration.GetDenominator());
 		Duration = MovieDuration.GetAsTimebase(ETimespan::TicksPerSecond);
 	}
-#endif
 
 	// If there are timecode tracks, find which tracks references them.
 	// If there are none, then - if there is only a single timecode track - we apply it to all other tracks.
@@ -605,7 +657,7 @@ void FElectraProtronPlayer::FImpl::InternalOpen(const FString& InFilename)
 		{
 			continue;
 		}
-		Tracks[i]->MP4Track = Electra::UtilitiesMP4::FMP4Track::Create(Tracks[i]->TrackBox);
+		Tracks[i]->MP4Track = Electra::UtilitiesMP4::FMP4Track::Create(Tracks[i]->TrackBox, Tracks[i]->FragmentInfo);
 		if (!Tracks[i]->MP4Track->Prepare(EntireMovieDuration, MovieDuration))
 		{
 			LastErrorMessage = Tracks[i]->MP4Track->GetLastError();
@@ -627,7 +679,12 @@ void FElectraProtronPlayer::FImpl::InternalOpen(const FString& InFilename)
 			// Read the first sample.
 			int64 SampleSize = It->GetSampleSize();
 			int64 SampleFileOffset = It->GetSampleFileOffset();
-			check(SampleSize == 4);		// If that triggers this is not a tmcd sample.
+			if (SampleSize != 4)
+			{
+				// Not a tmcd sample.
+				LastErrorMessage = FString::Printf(TEXT("Invalid timecode sample size."));
+				return;
+			}
 			TArray<uint32> TimecodeBuffer;
 			TimecodeBuffer.SetNumUninitialized(Align(SampleSize, sizeof(uint32)));
 			int64 NumRead = Reader->ReadData(TimecodeBuffer.GetData(), SampleSize, SampleFileOffset, Electra::IBaseDataReader::FCancellationCheckDelegate::CreateLambda([&](){return bAbort;}));
@@ -1016,7 +1073,7 @@ void FElectraProtronPlayer::FImpl::GetTrackCodecInfo(ElectraProtronUtils::FCodec
 
 FVariant FElectraProtronPlayer::FImpl::GetMediaInfo(FName InInfoName)
 {
-	if (InInfoName == ElectraProtronOptionNames::StartTimecodeValue || InInfoName == ElectraProtronOptionNames::StartTimecodeFrameRate || 
+	if (InInfoName == ElectraProtronOptionNames::StartTimecodeValue || InInfoName == ElectraProtronOptionNames::StartTimecodeFrameRate ||
 		InInfoName == ElectraProtronOptionNames::KeyframeInterval || InInfoName == ElectraProtronOptionNames::VideoDecoderProvider)
 	{
 		auto ci = CodecTypeIndex(ElectraProtronUtils::FCodecInfo::EType::Video);
diff --git a/Engine/Plugins/Media/ElectraPlayer/Source/ElectraProtron/Private/Player/ElectraProtronPlayerImpl.h b/Engine/Plugins/Media/ElectraPlayer/Source/ElectraProtron/Private/Player/ElectraProtronPlayerImpl.h
--- a/Engine/Plugins/Media/ElectraPlayer/Source/ElectraProtron/Private/Player/ElectraProtronPlayerImpl.h
+++ b/Engine/Plugins/Media/ElectraPlayer/Source/ElectraProtron/Private/Player/ElectraProtronPlayerImpl.h
@@ -350,6 +350,7 @@ private:
 		};
 		FDecoderProviderInfo DecoderProviderInfo;
 		TSharedPtr<Electra::UtilitiesMP4::FMP4BoxTRAK, ESPMode::ThreadSafe> TrackBox;
+		TSharedPtr<Electra::UtilitiesMP4::FMP4Track::FFragmentInfo, ESPMode::ThreadSafe> FragmentInfo;
 		TSharedPtr<Electra::UtilitiesMP4::FMP4Track, ESPMode::ThreadSafe> MP4Track;
 		TWeakPtr<FTrackInfo, ESPMode::ThreadSafe> ReferencedTimecodeTrack;
 		uint32 TrackID = 0;
@@ -435,8 +436,7 @@ private:
 	TQueue<FWorkerThreadMessage, EQueueMode::Mpsc> WorkMessages;
 	FString LastErrorMessage;
 	bool bAbort = false;
-	TArray<Electra::UtilitiesMP4::FMP4BoxTreeParser> ParsedRootBoxes;
-
+	TArray<TSharedPtr<Electra::UtilitiesMP4::FMP4BoxBase>> ParsedRootBoxes;
 	TArray<TSharedPtr<FTrackInfo, ESPMode::ThreadSafe>> Tracks;
 	TArray<TArray<int32>> UsableTrackArrayIndicesByType;
 
diff --git a/Engine/Plugins/Media/ElectraUtil/Source/ElectraBase/Private/Utilities/MP4Boxes/MP4Boxes.cpp b/Engine/Plugins/Media/ElectraUtil/Source/ElectraBase/Private/Utilities/MP4Boxes/MP4Boxes.cpp
--- a/Engine/Plugins/Media/ElectraUtil/Source/ElectraBase/Private/Utilities/MP4Boxes/MP4Boxes.cpp
+++ b/Engine/Plugins/Media/ElectraUtil/Source/ElectraBase/Private/Utilities/MP4Boxes/MP4Boxes.cpp
@@ -1383,8 +1383,12 @@ namespace Electra
 		struct FMP4BoxSTSZ::FParsed
 		{
 			const uint32* SizeArray32 = nullptr;
+			const uint16* SizeArray16 = nullptr;
+			const uint8* SizeArray8 = nullptr;
+			const uint8* SizeArray4 = nullptr;
 			uint32 NumberOfSamples = 0;
 			uint32 ConstantSampleSize = 0;
+			uint32 FieldSize = 0;
 			uint32 Flags = 0;
 			uint8 Version = 0;
 		};
@@ -1397,9 +1401,34 @@ namespace Electra
 			Parsed = MakePimpl<FParsed>();
 			FMP4AtomReader ar(BoxInfo.Data);
 			ar.ReadVersionAndFlags(Parsed->Version, Parsed->Flags);
-			ar.Read(Parsed->ConstantSampleSize);
-			ar.Read(Parsed->NumberOfSamples);
-			Parsed->SizeArray32 = (const uint32*) ar.GetCurrentDataPointer();
+			if (BoxInfo.Type == MakeBoxAtom('s','t','s','z'))
+			{
+				ar.Read(Parsed->ConstantSampleSize);
+				ar.Read(Parsed->NumberOfSamples);
+				if (Parsed->ConstantSampleSize == 0)
+				{
+					Parsed->SizeArray32 = (const uint32*) ar.GetCurrentDataPointer();
+				}
+			}
+			else
+			{
+				ar.Read(Parsed->FieldSize);
+				check((Parsed->FieldSize >> 8) == 0);
+				check(Parsed->FieldSize == 4 || Parsed->FieldSize == 8 || Parsed->FieldSize == 16);
+				ar.Read(Parsed->NumberOfSamples);
+				if (Parsed->FieldSize == 16)
+				{
+					Parsed->SizeArray16 = (const uint16*) ar.GetCurrentDataPointer();
+				}
+				else if (Parsed->FieldSize == 8)
+				{
+					Parsed->SizeArray8 = (const uint8*) ar.GetCurrentDataPointer();
+				}
+				else if (Parsed->FieldSize == 4)
+				{
+					Parsed->SizeArray4 = (const uint8*) ar.GetCurrentDataPointer();
+				}
+			}
 		}
 
 		uint32 FMP4BoxSTSZ::GetNumberOfSamples()
@@ -1410,7 +1439,32 @@ namespace Electra
 		uint32 FMP4BoxSTSZ::GetSizeOfSample(uint32 InIndex)
 		{
 			ParseIfRequired();
-			return Parsed->ConstantSampleSize ? Parsed->ConstantSampleSize : InIndex < Parsed->NumberOfSamples ? GetFromBigEndian(Parsed->SizeArray32[InIndex]) : 0;
+			if (InIndex >= Parsed->NumberOfSamples)
+			{
+				return 0;
+			}
+			if (Parsed->SizeArray32)
+			{
+				return GetFromBigEndian(Parsed->SizeArray32[InIndex]);
+			}
+			else if (Parsed->ConstantSampleSize)
+			{
+				return Parsed->ConstantSampleSize;
+			}
+			else if (Parsed->SizeArray16)
+			{
+				return GetFromBigEndian(Parsed->SizeArray16[InIndex]);
+			}
+			else if (Parsed->SizeArray8)
+			{
+				return Parsed->SizeArray8[InIndex];
+			}
+			else if (Parsed->SizeArray4)
+			{
+				uint8 sz = Parsed->SizeArray4[InIndex >> 1];
+				return (InIndex & 1) == 0 ? (sz >> 4) : (sz & 15);
+			}
+			return 0;
 		}
 
 /****************************************************************************************************************************************************/
@@ -1698,6 +1752,11 @@ namespace Electra
 			if (Parsed->Version == 1)
 			{
 				ar.Read(Parsed->FragmentDuration);
+				if (Parsed->FragmentDuration > 0x7fffffffffffffffULL)
+				{
+					Parsed->FragmentDuration = 0x7fffffffffffffffULL;
+					UE_LOG(LogElectraBase, Warning, TEXT("`mehd` version 1 fragment duration cannot be represented as a signed value. Why is it so large?"));
+				}
 			}
 			else
 			{
@@ -1707,10 +1766,10 @@ namespace Electra
 			}
 		}
 
-		uint64 FMP4BoxMEHD::GetFragmentDuration()
+		int64 FMP4BoxMEHD::GetFragmentDuration()
 		{
 			ParseIfRequired();
-			return Parsed->FragmentDuration;
+			return (int64) Parsed->FragmentDuration;
 		}
 
 /****************************************************************************************************************************************************/
@@ -1954,6 +2013,7 @@ namespace Electra
 			TArray<uint32> SampleSizes;
 			TArray<uint32> SampleFlags;
 			TArray<int32> SampleCompositionTimeOffsets;
+			int64 TotalDuration = 0;
 			uint32 SampleCount = 0;
 			uint32 FirstSampleFlags = 0;
 			uint32 Flags = 0;
@@ -2000,7 +2060,10 @@ namespace Electra
 				{
 					if ((Parsed->Flags & 0x100) != 0)
 					{
-						ar.Read(Parsed->SampleDurations[i]);
+						uint32 dur;
+						ar.Read(dur);
+						Parsed->SampleDurations[i] = dur;
+						Parsed->TotalDuration += dur;
 					}
 					if ((Parsed->Flags & 0x200) != 0)
 					{
@@ -2104,6 +2167,11 @@ namespace Electra
 			ParseIfRequired();
 			return Parsed->SampleCompositionTimeOffsets;
 		}
+		int64 FMP4BoxTRUN::GetTotalDuration()
+		{
+			ParseIfRequired();
+			return Parsed->TotalDuration;
+		}
 
 /****************************************************************************************************************************************************/
 
@@ -2502,7 +2570,7 @@ namespace Electra
 			FactoryMap.Add(MakeBoxAtom('m','d','i','a'), &FMP4BoxMDIA::Create);
 			FactoryMap.Add(MakeBoxAtom('m','d','h','d'), &FMP4BoxMDHD::Create);
 			FactoryMap.Add(MakeBoxAtom('h','d','l','r'), &FMP4BoxHDLR::Create);
-			FactoryMap.Add(MakeBoxAtom('e','l','n','g'), &FMP4BoxHDLR::Create);
+			FactoryMap.Add(MakeBoxAtom('e','l','n','g'), &FMP4BoxELNG::Create);
 			FactoryMap.Add(MakeBoxAtom('m','i','n','f'), &FMP4BoxMINF::Create);
 			FactoryMap.Add(MakeBoxAtom('d','i','n','f'), &FMP4BoxDINF::Create);
 			FactoryMap.Add(MakeBoxAtom('d','r','e','f'), &FMP4BoxDREF::Create);
@@ -2515,6 +2583,7 @@ namespace Electra
 			FactoryMap.Add(MakeBoxAtom('s','d','t','p'), &FMP4BoxSDTP::Create);
 			FactoryMap.Add(MakeBoxAtom('s','t','s','c'), &FMP4BoxSTSC::Create);
 			FactoryMap.Add(MakeBoxAtom('s','t','s','z'), &FMP4BoxSTSZ::Create);
+			FactoryMap.Add(MakeBoxAtom('s','t','z','2'), &FMP4BoxSTSZ::Create);
 			FactoryMap.Add(MakeBoxAtom('s','t','c','o'), &FMP4BoxSTCO::Create);
 			FactoryMap.Add(MakeBoxAtom('c','o','6','4'), &FMP4BoxSTCO::Create);
 			FactoryMap.Add(MakeBoxAtom('s','a','i','z'), &FMP4BoxSAIZ::Create);
diff --git a/Engine/Plugins/Media/ElectraUtil/Source/ElectraBase/Private/Utilities/MP4Boxes/MP4Track.cpp b/Engine/Plugins/Media/ElectraUtil/Source/ElectraBase/Private/Utilities/MP4Boxes/MP4Track.cpp
--- a/Engine/Plugins/Media/ElectraUtil/Source/ElectraBase/Private/Utilities/MP4Boxes/MP4Track.cpp
+++ b/Engine/Plugins/Media/ElectraUtil/Source/ElectraBase/Private/Utilities/MP4Boxes/MP4Track.cpp
@@ -8,6 +8,142 @@ namespace Electra
 	namespace UtilitiesMP4
 	{
 
+		bool FMP4Track::FFragmentInfo::Prepare(const TSharedPtr<FMP4BoxTRAK, ESPMode::ThreadSafe>& InTrakBox, const TSharedPtr<FMP4BoxMVEX, ESPMode::ThreadSafe>& InMvexBox, const TArray<TSharedPtr<FMP4BoxBase>>& InAllRootBoxes)
+		{
+			if (bHasBeenPrepared)
+			{
+				return true;
+			}
+			if (!InTrakBox || !InMvexBox || InAllRootBoxes.IsEmpty())
+			{
+				return false;
+			}
+			TkhdBox = InTrakBox->FindBoxRecursive<FMP4BoxTKHD>(MakeBoxAtom('t','k','h','d'), 0);
+			MdhdBox = InTrakBox->FindBoxRecursive<FMP4BoxMDHD>(MakeBoxAtom('m','d','h','d'), 1);
+			auto _StszBox = InTrakBox->FindBoxRecursive<FMP4BoxSTSZ>(MakeBoxAtom('s','t','s','z'), 5);
+
+			check(TkhdBox && MdhdBox && _StszBox);
+			if (!TkhdBox || !MdhdBox || !_StszBox)
+			{
+				return false;
+			}
+			// Get the ID of this track.
+			const uint32 TrackID = TkhdBox->GetTrackID();
+			// Locate the mandatory `trex` box.
+			TrexBox = InMvexBox->FindBoxRecursiveByPredicate<FMP4BoxTREX>(MakeBoxAtom('t','r','e','x'),[tkid=TrackID](const TSharedPtr<FMP4BoxTREX, ESPMode::ThreadSafe>& InTrex)
+					{ return InTrex->GetTrackID() == tkid; }, 0);
+			if (!TrexBox)
+			{
+				return false;
+			}
+			// Locate all `moof` boxes containing track data for this track and collect their `traf` boxes.
+			TArray<TSharedPtr<FMP4BoxTRAF, ESPMode::ThreadSafe>> TrackTrafBoxes;
+			for(const TSharedPtr<FMP4BoxBase>& rb : InAllRootBoxes)
+			{
+				if (rb->GetType() != MakeBoxAtom('m','o','o','f'))
+				{
+					continue;
+				}
+				rb->GetAllBoxInstancesByPredicate(TrackTrafBoxes, MakeBoxAtom('t','r','a','f'), [tkid=TrackID](const TSharedPtr<FMP4BoxTRAF, ESPMode::ThreadSafe>& InTraf)
+				{
+					auto Tfhd = InTraf->FindBoxRecursive<FMP4BoxTFHD>(MakeBoxAtom('t','f','h','d'), 0);
+					return Tfhd && Tfhd->GetTrackID() == tkid;
+				});
+			}
+
+			// How many samples are in the `moov`?
+			NumBaseTrackSamples = _StszBox->GetNumberOfSamples();
+
+			// Set up a list of the matching `traf` boxes and their related info boxes for convenience.
+			if (TrackTrafBoxes.Num())
+			{
+				NumTotalFragmentSamples = 0;
+				TrackFragments.SetNum(TrackTrafBoxes.Num());
+				for(int32 i=0, iMax=TrackTrafBoxes.Num(); i<iMax; ++i)
+				{
+					auto Frg = MakeShared<FTrackTrafs, ESPMode::ThreadSafe>();
+					Frg->TrafBox = TrackTrafBoxes[i];
+					Frg->MoofBox = Frg->TrafBox->FindParentBox<FMP4BoxMOOF>(MakeBoxAtom('m','o','o','f'));
+					Frg->TfhdBox = Frg->TrafBox->FindBoxRecursive<FMP4BoxTFHD>(MakeBoxAtom('t','f','h','d'), 0);
+					Frg->TrafBox->GetAllBoxInstances<FMP4BoxTRUN>(Frg->TrunBoxes, MakeBoxAtom('t','r','u','n'));
+					Frg->AbsoluteFirstSampleNumber = NumBaseTrackSamples + NumTotalFragmentSamples;
+
+					Frg->NumSamples = 0;
+					for(int32 j=0, jMax=Frg->TrunBoxes.Num(); j<jMax; ++j)
+					{
+						Frg->NumSamples += Frg->TrunBoxes[j]->GetNumberOfSamples();
+					}
+					NumTotalFragmentSamples += Frg->NumSamples;
+					TrackFragments[i] = MoveTemp(Frg);
+				}
+			}
+			bHasBeenPrepared = true;
+			return true;
+		}
+
+		void FMP4Track::FFragmentInfo::CalculateDurations()
+		{
+			if (!bDidCalculateDuration)
+			{
+				bDidCalculateDuration = true;
+
+				BaseTrackDuration = 0;
+				auto Mdia = MdhdBox->GetParentBox();
+				check(Mdia.IsValid());
+				if (Mdia.IsValid())
+				{
+					auto Stts = Mdia->FindBoxRecursive<FMP4BoxSTTS>(MakeBoxAtom('s','t','t','s'), 5);
+					check(Stts.IsValid());
+					if (Stts.IsValid())
+					{
+						BaseTrackDuration = Stts->GetTotalDuration();
+					}
+				}
+
+				SumOfFragmentDurations = 0;
+				for(const TSharedPtr<FTrackTrafs, ESPMode::ThreadSafe>& tf : TrackFragments)
+				{
+					tf->Duration = 0;
+					for(const TSharedPtr<FMP4BoxTRUN, ESPMode::ThreadSafe>& tr : tf->TrunBoxes)
+					{
+						if (tr->HasSampleDurations())
+						{
+							tf->Duration += tr->GetTotalDuration();
+						}
+						else
+						{
+							if (tf->TfhdBox->HasDefaultSampleDuration())
+							{
+								tf->Duration += tr->GetNumberOfSamples() * tf->TfhdBox->GetDefaultSampleDuration();
+							}
+							else
+							{
+								tf->Duration += tr->GetNumberOfSamples() * TrexBox->GetDefaultSampleDuration();
+							}
+						}
+					}
+					tf->FirstTrackTotalDTS = BaseTrackDuration + SumOfFragmentDurations;
+					tf->LastTrackTotalDTSAndDur = tf->FirstTrackTotalDTS + tf->Duration;
+					SumOfFragmentDurations += tf->Duration;
+				}
+			}
+		}
+
+		int64 FMP4Track::FFragmentInfo::GetTotalDuration()
+		{
+			CalculateDurations();
+			return BaseTrackDuration + SumOfFragmentDurations;
+		}
+
+
+
+
+
+
+
+
+
+
 		bool FMP4Track::Prepare(FTimeFraction InFullMovieDuration, FTimeFraction InAdjustedMovieDuration)
 		{
 			if (!TrakBox.IsValid())
@@ -15,7 +151,7 @@ namespace Electra
 				LastErrorMessage = FString::Printf(TEXT("No `trak` box given."));
 				return false;
 			}
-			TkhdBox = TrakBox->FindBoxRecursive<Electra::UtilitiesMP4::FMP4BoxTKHD>(Electra::UtilitiesMP4::MakeBoxAtom('t','k','h','d'), 0);
+			TkhdBox = TrakBox->FindBoxRecursive<FMP4BoxTKHD>(MakeBoxAtom('t','k','h','d'), 0);
 			if (!TkhdBox.IsValid())
 			{
 				LastErrorMessage = FString::Printf(TEXT("No `tkhd` box in `trak`."));
@@ -24,29 +160,28 @@ namespace Electra
 			Convs.TrackID = TkhdBox->GetTrackID();
 			Convs.FullMovieDuration = InFullMovieDuration;
 			Convs.DurationFromMvhdBox = InAdjustedMovieDuration;
-			Convs.DurationFromTkhdBox.SetFromND(TkhdBox->GetDuration(), Convs.DurationFromMvhdBox.GetDenominator());
 
 			// Check for correct box hierarchy.
-			TSharedPtr<FMP4BoxMDIA, ESPMode::ThreadSafe> MdiaBox = TrakBox->FindBoxRecursive<Electra::UtilitiesMP4::FMP4BoxMDIA>(Electra::UtilitiesMP4::MakeBoxAtom('m','d','i','a'), 0);
+			TSharedPtr<FMP4BoxMDIA, ESPMode::ThreadSafe> MdiaBox = TrakBox->FindBoxRecursive<FMP4BoxMDIA>(MakeBoxAtom('m','d','i','a'), 0);
 			if (!MdiaBox.IsValid())
 			{
 				LastErrorMessage = FString::Printf(TEXT("No `mdia` box in `trak`."));
 				return false;
 			}
-			TSharedPtr<FMP4BoxMINF, ESPMode::ThreadSafe> MinfBox = MdiaBox->FindBoxRecursive<Electra::UtilitiesMP4::FMP4BoxMINF>(Electra::UtilitiesMP4::MakeBoxAtom('m','i','n','f'), 0);
+			TSharedPtr<FMP4BoxMINF, ESPMode::ThreadSafe> MinfBox = MdiaBox->FindBoxRecursive<FMP4BoxMINF>(MakeBoxAtom('m','i','n','f'), 0);
 			if (!MinfBox.IsValid())
 			{
 				LastErrorMessage = FString::Printf(TEXT("No `minf` box in `mdia`."));
 				return false;
 			}
-			TSharedPtr<FMP4BoxSTBL, ESPMode::ThreadSafe> StblBox = MinfBox->FindBoxRecursive<Electra::UtilitiesMP4::FMP4BoxSTBL>(Electra::UtilitiesMP4::MakeBoxAtom('s','t','b','l'), 0);
+			TSharedPtr<FMP4BoxSTBL, ESPMode::ThreadSafe> StblBox = MinfBox->FindBoxRecursive<FMP4BoxSTBL>(MakeBoxAtom('s','t','b','l'), 0);
 			if (!StblBox.IsValid())
 			{
 				LastErrorMessage = FString::Printf(TEXT("No `stbl` box in `minf`."));
 				return false;
 			}
 
-			MdhdBox = MdiaBox->FindBoxRecursive<Electra::UtilitiesMP4::FMP4BoxMDHD>(Electra::UtilitiesMP4::MakeBoxAtom('m','d','h','d'), 0);
+			MdhdBox = MdiaBox->FindBoxRecursive<FMP4BoxMDHD>(MakeBoxAtom('m','d','h','d'), 0);
 			if (!MdhdBox.IsValid())
 			{
 				LastErrorMessage = FString::Printf(TEXT("No `mdhd` box in `mdia`."));
@@ -58,27 +193,44 @@ namespace Electra
 				LastErrorMessage = FString::Printf(TEXT("Timescale in `mdhd` box is zero, which is not supported."));
 				return false;
 			}
+			// If this is a fragmented file then the duration in the mdhd box does not necessarily describe the entire media.
+			// Get it from the calculated duration of all fragments.
+			if (FragmentInfo.IsValid())
+			{
+				check(FragmentInfo->HasBeenPrepared());
+				if (!FragmentInfo->HasBeenPrepared())
+				{
+					LastErrorMessage = FString::Printf(TEXT("Fragment information has not been set up."));
+					return false;
+				}
+				Convs.DurationFromTkhdBox.SetFromND(FragmentInfo->GetTotalDuration(), Convs.DurationFromMdhdBox.GetDenominator());
+				Convs.DurationFromMdhdBox.SetNumerator(FragmentInfo->GetTotalDuration());
+			}
+			else
+			{
+				Convs.DurationFromTkhdBox.SetFromND(TkhdBox->GetDuration(), Convs.DurationFromMvhdBox.GetDenominator());
+			}
 
 			// Required sample information boxes:
-			SttsBox = StblBox->FindBoxRecursive<Electra::UtilitiesMP4::FMP4BoxSTTS>(Electra::UtilitiesMP4::MakeBoxAtom('s','t','t','s'), 0);
+			SttsBox = StblBox->FindBoxRecursive<FMP4BoxSTTS>(MakeBoxAtom('s','t','t','s'), 0);
 			if (!SttsBox.IsValid())
 			{
 				LastErrorMessage = FString::Printf(TEXT("No `stts` box in `stbl`."));
 				return false;
 			}
-			StscBox = StblBox->FindBoxRecursive<Electra::UtilitiesMP4::FMP4BoxSTSC>(Electra::UtilitiesMP4::MakeBoxAtom('s','t','s','c'), 0);
+			StscBox = StblBox->FindBoxRecursive<FMP4BoxSTSC>(MakeBoxAtom('s','t','s','c'), 0);
 			if (!StscBox.IsValid())
 			{
 				LastErrorMessage = FString::Printf(TEXT("No `stsc` box in `stbl`."));
 				return false;
 			}
-			StszBox = StblBox->FindBoxRecursive<Electra::UtilitiesMP4::FMP4BoxSTSZ>(Electra::UtilitiesMP4::MakeBoxAtom('s','t','s','z'), 0);
+			StszBox = StblBox->FindBoxRecursive<FMP4BoxSTSZ>(MakeBoxAtom('s','t','s','z'), 0);
 			if (!StszBox.IsValid())
 			{
 				LastErrorMessage = FString::Printf(TEXT("No `stsz` box in `stbl`."));
 				return false;
 			}
-			StcoBox = StblBox->FindBoxRecursive<Electra::UtilitiesMP4::FMP4BoxSTCO>(Electra::UtilitiesMP4::MakeBoxAtom('s','t','c','o'), 0);
+			StcoBox = StblBox->FindBoxRecursive<FMP4BoxSTCO>(MakeBoxAtom('s','t','c','o'), 0);
 			if (!StcoBox.IsValid())
 			{
 				LastErrorMessage = FString::Printf(TEXT("No `stco` ot `co64` box in `stbl`."));
@@ -86,25 +238,38 @@ namespace Electra
 			}
 
 			// Validity check
-			Convs.NumTotalSamples = StszBox->GetNumberOfSamples();
-			if (Convs.NumTotalSamples != SttsBox->GetNumTotalSamples())
+			Convs.NumMoovSamples = StszBox->GetNumberOfSamples();
+			if (Convs.NumMoovSamples != SttsBox->GetNumTotalSamples())
 			{
 				LastErrorMessage = FString::Printf(TEXT("Mismatching number of samples in `stts` and `stsz` boxes."));
 				return false;
 			}
+			Convs.NumTotalSamples = Convs.NumMoovSamples + (FragmentInfo.IsValid() ? FragmentInfo->NumTotalFragmentSamples : 0);
 
 			// Optional sample information boxes:
-			CttsBox = StblBox->FindBoxRecursive<Electra::UtilitiesMP4::FMP4BoxCTTS>(Electra::UtilitiesMP4::MakeBoxAtom('c','t','t','s'), 0);
-			StssBox = StblBox->FindBoxRecursive<Electra::UtilitiesMP4::FMP4BoxSTSS>(Electra::UtilitiesMP4::MakeBoxAtom('s','t','s','s'), 0);
-			StblBox->GetAllBoxInstances(SgpdBoxes, Electra::UtilitiesMP4::MakeBoxAtom('s','g','p','d'));
-			StblBox->GetAllBoxInstances(SbgpBoxes, Electra::UtilitiesMP4::MakeBoxAtom('s','b','g','p'));
+			CttsBox = StblBox->FindBoxRecursive<FMP4BoxCTTS>(MakeBoxAtom('c','t','t','s'), 0);
+			StssBox = StblBox->FindBoxRecursive<FMP4BoxSTSS>(MakeBoxAtom('s','t','s','s'), 0);
+			StblBox->GetAllBoxInstances(SgpdBoxes, MakeBoxAtom('s','g','p','d'));
+			StblBox->GetAllBoxInstances(SbgpBoxes, MakeBoxAtom('s','b','g','p'));
 
 			// Start with default values for what is mapped onto the timeline.
-			Convs.CompositionTimeAtZeroPoint = CttsBox.IsValid() && CttsBox->GetEntries().Num() ? CttsBox->GetEntries()[0].sample_offset : 0;
-			Convs.MappedDurationFromElstBox = Convs.DurationFromMvhdBox;
+			Convs.CompositionTimeAtZeroPoint = 0;
+			if (CttsBox.IsValid() && CttsBox->GetEntries().Num())
+			{
+				Convs.CompositionTimeAtZeroPoint = CttsBox->GetEntries()[0].sample_offset;
+			}
+			else if (FragmentInfo && FragmentInfo->TrackFragments.Num() && FragmentInfo->TrackFragments[0]->TrunBoxes.Num())
+			{
+				auto Trun = FragmentInfo->TrackFragments[0]->TrunBoxes[0];
+				if (Trun->GetNumberOfSamples() && Trun->HasSampleCompositionTimeOffsets())
+				{
+					Convs.CompositionTimeAtZeroPoint = Trun->GetSampleCompositionTimeOffsets()[0];
+				}
+			}
 
+			Convs.MappedDurationFromElstBox = Convs.DurationFromMvhdBox;
 			// Optional edit list
-			ElstBox = TrakBox->FindBoxRecursive<Electra::UtilitiesMP4::FMP4BoxELST>(Electra::UtilitiesMP4::MakeBoxAtom('e','l','s','t'), 1);
+			ElstBox = TrakBox->FindBoxRecursive<FMP4BoxELST>(MakeBoxAtom('e','l','s','t'), 1);
 			// If there is an edit list it needs to be simple and only contain a composition time mapping.
 			if (ElstBox.IsValid())
 			{
@@ -115,7 +280,7 @@ namespace Electra
 				}
 				// If there is more than a single entry things are bound to be complicated, so we don't even
 				// want to know if the entry we are interested in is there.
-				const TArray<Electra::UtilitiesMP4::FMP4BoxELST::FEntry>& ElstEntries = ElstBox->GetEntries();
+				const TArray<FMP4BoxELST::FEntry>& ElstEntries = ElstBox->GetEntries();
 				if (ElstEntries.Num() == 0)
 				{
 					LastErrorMessage = FString::Printf(TEXT("Edit list is empty."));
@@ -159,27 +324,34 @@ namespace Electra
 				{
 					// The value of 0 is reserved for fragmented files with no `mehd` box.
 					// We do not handle this case.
-					LastErrorMessage = FString::Printf(TEXT("Edit list specifies zero edit duration, which is not supported."));
-					return false;
+					if (!FragmentInfo.IsValid())
+					{
+						LastErrorMessage = FString::Printf(TEXT("Edit list specifies zero edit duration, which is not supported."));
+						return false;
+					}
+					// Otherwise the mapped duration is already set as the entire movie duration.
+				}
+				else
+				{
+					Convs.MappedDurationFromElstBox.SetFromND(ElstEntries[0].edit_duration, Convs.DurationFromMvhdBox.GetDenominator());
 				}
-				Convs.MappedDurationFromElstBox.SetFromND(ElstEntries[0].edit_duration, Convs.DurationFromMvhdBox.GetDenominator());
 			}
 			// For convenience sake convert the mapped duration from `mvhd` timescale into `mdhd` timescale.
 			Convs.MappedDurationFromElstBox.SetFromND(Convs.MappedDurationFromElstBox.GetAsTimebase(Convs.DurationFromMdhdBox.GetDenominator()), Convs.DurationFromMdhdBox.GetDenominator());
 
 
 			// Set up common track meta data
-			UdtaBox = TrakBox->FindBoxRecursive<Electra::UtilitiesMP4::FMP4BoxUDTA>(Electra::UtilitiesMP4::MakeBoxAtom('u','d','t','a'), 0);
+			UdtaBox = TrakBox->FindBoxRecursive<FMP4BoxUDTA>(MakeBoxAtom('u','d','t','a'), 0);
 			if (UdtaBox.IsValid())
 			{
 				// Is there a `name` box?
-				TSharedPtr<FMP4BoxBase, ESPMode::ThreadSafe> NameBox = UdtaBox->FindBoxRecursive<Electra::UtilitiesMP4::FMP4BoxBase>(Electra::UtilitiesMP4::MakeBoxAtom('n','a','m','e'), 0);
+				TSharedPtr<FMP4BoxBase, ESPMode::ThreadSafe> NameBox = UdtaBox->FindBoxRecursive<FMP4BoxBase>(MakeBoxAtom('n','a','m','e'), 0);
 				if (NameBox.IsValid() && NameBox->GetBoxData().Num())
 				{
 					CommonMetadata.Name = FString::ConstructFromPtrSize(reinterpret_cast<const ANSICHAR*>(NameBox->GetBoxData().GetData()), NameBox->GetBoxData().Num());
 				}
 			}
-			TSharedPtr<FMP4BoxHDLR, ESPMode::ThreadSafe> HdlrBox = MdiaBox->FindBoxRecursive<Electra::UtilitiesMP4::FMP4BoxHDLR>(Electra::UtilitiesMP4::MakeBoxAtom('h','d','l','r'), 0);
+			TSharedPtr<FMP4BoxHDLR, ESPMode::ThreadSafe> HdlrBox = MdiaBox->FindBoxRecursive<FMP4BoxHDLR>(MakeBoxAtom('h','d','l','r'), 0);
 			check(HdlrBox.IsValid())
 			if (HdlrBox.IsValid())
 			{
@@ -188,7 +360,7 @@ namespace Electra
 			// Get language from `mdhd` box first.
 			CommonMetadata.LanguageTag = MdhdBox->GetLanguageTag();
 			// Then, if there is an `elng` box that gives better information get it from there.
-			TSharedPtr<FMP4BoxELNG, ESPMode::ThreadSafe> ElngBox = MdiaBox->FindBoxRecursive<Electra::UtilitiesMP4::FMP4BoxELNG>(Electra::UtilitiesMP4::MakeBoxAtom('e','l','n','g'), 0);
+			TSharedPtr<FMP4BoxELNG, ESPMode::ThreadSafe> ElngBox = MdiaBox->FindBoxRecursive<FMP4BoxELNG>(MakeBoxAtom('e','l','n','g'), 0);
 			if (ElngBox.IsValid())
 			{
 				CommonMetadata.LanguageTag = ElngBox->GetLanguageTag();
@@ -204,7 +376,7 @@ namespace Electra
 			// which is not handled.
 			if (!StartIt.IsValid())
 			{
-				LastErrorMessage = FString::Printf(TEXT("Could not locate first media sample. Is this an empty or a fragmented mp4?"));
+				LastErrorMessage = FString::Printf(TEXT("Could not locate first media sample. Is this an empty mp4?"));
 				return false;
 			}
 			const int64 pts0 = Convs.CompositionTimeAtZeroPoint;
@@ -355,28 +527,32 @@ namespace Electra
 			}
 			TUniquePtr<FIterator> It(new FIterator);
 			It->Track = AsWeak();
-			It->SampleNumber = 0;
+			It->FragmentVars.TrunIterator.SetFragmentInfo(FragmentInfo);
 			// Copy convenience values into the iterator
 			It->Convs = Convs;
 			// Set up the box iterators.
-			It->stszIt.SetBox(StszBox);
-			It->sttsIt.SetBox(SttsBox);
-			It->cttsIt.SetBox(CttsBox, Convs.NumTotalSamples);
-			It->stcoIt.SetBox(StcoBox);
-			It->stscIt.SetBox(StscBox, Convs.NumTotalSamples);
-			It->stssIt.SetBox(StssBox, Convs.NumTotalSamples);
-			// Do we have a `rap ` group?
-			TSharedPtr<FMP4BoxSGPD, ESPMode::ThreadSafe>* rapSGPD = SgpdBoxes.FindByPredicate([](const TSharedPtr<FMP4BoxSGPD, ESPMode::ThreadSafe>& e){ return e->GetGroupingType() == Electra::UtilitiesMP4::MakeBoxAtom('r','a','p',' ');});
-			TSharedPtr<FMP4BoxSBGP, ESPMode::ThreadSafe>* rapSBGP = SbgpBoxes.FindByPredicate([](const TSharedPtr<FMP4BoxSBGP, ESPMode::ThreadSafe>& e){ return e->GetGroupingType() == Electra::UtilitiesMP4::MakeBoxAtom('r','a','p',' ');});
-			if (rapSGPD && rapSBGP)
-			{
-				It->rapIt.SetBox(*rapSBGP, (*rapSGPD)->GetDefaultGroupDescriptionIndex(), Convs.NumTotalSamples);
-			}
-			else
+			if (Convs.NumMoovSamples)
 			{
-				// Initialize the iterator such that it can be used to return "not a RAP" for every sample.
-				It->rapIt.SetBox(nullptr, 0, Convs.NumTotalSamples);
+				It->stszIt.SetBox(StszBox);
+				It->sttsIt.SetBox(SttsBox);
+				It->cttsIt.SetBox(CttsBox, Convs.NumMoovSamples);
+				It->stcoIt.SetBox(StcoBox);
+				It->stscIt.SetBox(StscBox, Convs.NumMoovSamples);
+				It->stssIt.SetBox(StssBox, Convs.NumMoovSamples);
+				// Do we have a `rap ` group?
+				TSharedPtr<FMP4BoxSGPD, ESPMode::ThreadSafe>* rapSGPD = SgpdBoxes.FindByPredicate([](const TSharedPtr<FMP4BoxSGPD, ESPMode::ThreadSafe>& e){ return e->GetGroupingType() == MakeBoxAtom('r','a','p',' ');});
+				TSharedPtr<FMP4BoxSBGP, ESPMode::ThreadSafe>* rapSBGP = SbgpBoxes.FindByPredicate([](const TSharedPtr<FMP4BoxSBGP, ESPMode::ThreadSafe>& e){ return e->GetGroupingType() == MakeBoxAtom('r','a','p',' ');});
+				if (rapSGPD && rapSBGP)
+				{
+					It->rapIt.SetBox(*rapSBGP, (*rapSGPD)->GetDefaultGroupDescriptionIndex(), Convs.NumMoovSamples);
+				}
+				else
+				{
+					// Initialize the iterator such that it can be used to return "not a RAP" for every sample.
+					It->rapIt.SetBox(nullptr, 0, Convs.NumMoovSamples);
+				}
 			}
+			It->StartAt(0);
 			It->Update();
 			return MakeShareable(It.Release());
 		}
@@ -403,15 +579,9 @@ namespace Electra
 			TSharedPtr<FMP4Track::FIterator, ESPMode::ThreadSafe> It(CreateIterator());
 			if (It.IsValid() && InAtSampleNumber)
 			{
+				// Clamp to last sample if passed number is too large.
 				InAtSampleNumber = InAtSampleNumber < Convs.NumTotalSamples ? InAtSampleNumber : Convs.NumTotalSamples-1;
-				It->SampleNumber = InAtSampleNumber;
-				It->stszIt.SetToSampleNumber(InAtSampleNumber);
-				It->sttsIt.SetToSampleNumber(InAtSampleNumber);
-				It->cttsIt.SetToSampleNumber(InAtSampleNumber);
-				It->stcoIt.SetToSampleNumber(InAtSampleNumber);
-				It->stscIt.SetToSampleNumber(InAtSampleNumber);
-				It->stssIt.SetToSampleNumber(InAtSampleNumber);
-				It->rapIt.SetToSampleNumber(InAtSampleNumber);
+				It->StartAt(InAtSampleNumber);
 				It->Update();
 			}
 			return It;
@@ -421,7 +591,7 @@ namespace Electra
 		TSharedPtr<FMP4Track::FIterator, ESPMode::ThreadSafe> FMP4Track::CreateIteratorAtKeyframe(FTimeValue InForTime, FTimeValue InLaterTimeThreshold)
 		{
 			check(bHasBeenPrepared);
-			if (!bHasBeenPrepared || !SttsBox.IsValid() || SttsBox->GetTotalDuration()==0)
+			if (!bHasBeenPrepared || !SttsBox.IsValid())
 			{
 				LastErrorMessage = FString::Printf(TEXT("Track has not been prepared, cannot create an iterator."));
 				return nullptr;
@@ -431,7 +601,8 @@ namespace Electra
 				LastErrorMessage = FString::Printf(TEXT("Invalid time, cannot create an iterator."));
 				return nullptr;
 			}
-			if (!Convs.DurationFromMdhdBox.IsValid() || !Convs.MappedDurationFromElstBox.IsValid())
+			const int64 TotalDuration = FragmentInfo.IsValid() ? FragmentInfo->GetTotalDuration() : SttsBox->GetTotalDuration();
+			if (TotalDuration <=0 || !Convs.DurationFromMdhdBox.IsValid() || !Convs.MappedDurationFromElstBox.IsValid())
 			{
 				LastErrorMessage = FString::Printf(TEXT("Invalid track duration, cannot create an iterator."));
 				return nullptr;
@@ -456,7 +627,7 @@ namespace Electra
 			LocalTrackTime += Convs.CompositionTimeAtZeroPoint;
 			MaxLocalTrackTime += Convs.CompositionTimeAtZeroPoint;
 
-			int64 ApproxSampleNumber = Convs.NumTotalSamples ? LocalTrackTime * (Convs.NumTotalSamples-1) / SttsBox->GetTotalDuration() : 0;
+			int64 ApproxSampleNumber = Convs.NumTotalSamples ? LocalTrackTime * (Convs.NumTotalSamples-1) / TotalDuration : 0;
 			check(ApproxSampleNumber <= 0xffffffff);
 
 
@@ -489,7 +660,7 @@ namespace Electra
 			{
 				// Due to possible frame reordering we need to look at the DTS here with the composition offset applied
 				// to be sure to find the correct sample. If we were to look at the PTS we could leave the loop to early.
-				for(int64 MaxDtsWithComp=MaxLocalTrackTime+Convs.CompositionTimeAtZeroPoint; !NextSyncFrameIt->IsLast() && (NextSyncFrameIt->GetDTS().GetNumerator() <= MaxDtsWithComp && !NextSyncFrameIt->IsSyncOrRAPSample()); NextSyncFrameIt->Next())
+				for(; !NextSyncFrameIt->IsLast() && (NextSyncFrameIt->GetDTS().GetNumerator() <= MaxLocalTrackTime && !NextSyncFrameIt->IsSyncOrRAPSample()); NextSyncFrameIt->Next())
 				{ }
 				bLaterOneIsPossible = NextSyncFrameIt->IsSyncOrRAPSample() && NextSyncFrameIt->GetPTS().GetNumerator() <= MaxLocalTrackTime;
 			}
@@ -511,51 +682,105 @@ namespace Electra
 /****************************************************************************************************************************************************/
 /****************************************************************************************************************************************************/
 
+		void FMP4Track::FIterator::StartAt(uint32 InFrameNumber)
+		{
+			SampleNumber = InFrameNumber;
+			if (InFrameNumber < Convs.NumMoovSamples)
+			{
+				stszIt.SetToSampleNumber(InFrameNumber);
+				sttsIt.SetToSampleNumber(InFrameNumber);
+				cttsIt.SetToSampleNumber(InFrameNumber);
+				stcoIt.SetToSampleNumber(InFrameNumber);
+				stscIt.SetToSampleNumber(InFrameNumber);
+				stssIt.SetToSampleNumber(InFrameNumber);
+				rapIt.SetToSampleNumber(InFrameNumber);
+				FragmentVars.bInFragments = false;
+			}
+			else
+			{
+				FragmentVars.TrunIterator.StartAt(InFrameNumber);
+				FragmentVars.bInFragments = true;
+			}
+		}
+
 		void FMP4Track::FIterator::Update()
 		{
 			if (!IsValid())
 			{
 				return;
 			}
-			int64 DTS = sttsIt.GetCurrentTime() + Convs.BaseMediaDecodeTime;
-			uint32 Duration = sttsIt.GetCurrentDuration();
-			int64 CompositionTimeOffset = cttsIt.GetCurrentOffset();
-			int64 PTS = DTS + CompositionTimeOffset;
-
-			const uint32 Timescale = Convs.DurationFromMdhdBox.GetDenominator();
-			CurrentDTS.SetFromND(DTS, Timescale);
-			CurrentPTS.SetFromND(PTS, Timescale);
-			CurrentEffectiveDTS.SetFromND(DTS - Convs.DTSShiftAtZeroPoint, Timescale);
-			CurrentEffectivePTS.SetFromND(PTS - Convs.CompositionTimeAtZeroPoint, Timescale);
-			// Set the duration as the fraction of the duration and the timescale.
-			CurrentDuration.SetFromND(Duration, Timescale);
-			// Also set the duration as the delta of the DTS of this sample and the next
-			// in timespan units. This is to avoid transformation issues from media local
-			// time into the timescale used in engine.
-			CurrentDurationTS = FTimeFraction(DTS + Duration, Timescale).GetAsTimespan() - CurrentDTS.GetAsTimespan();
-			CurrentSampleSize = stszIt.GetCurrentSampleSize();
-			bCurrentIsSyncOrRAP = stssIt.IsSyncSample() || rapIt.GetCurrentGroupDescriptionIndex() != 0;
-
-			// Which chunk is this sample in?
-			uint32 chunkIndex = stscIt.GetCurrentChunkIndex();
-			check(chunkIndex);
-			uint64 chunkOffset = stcoIt.GetOffsetForChunkIndex(chunkIndex - 1);
-			check(chunkOffset);
-			// Which sample position within the current chunk run are we at?
-			uint32 samplePosInChunk = stscIt.GetSampleIndexInCurrentChunk();
-			// Giving us which sample number at the start of the chunk?
-			uint32 sampleNumAtChunkStart = SampleNumber - samplePosInChunk;
-			for(uint32 i=0; i<samplePosInChunk; ++i)
-			{
-				chunkOffset += stszIt.GetSampleSizeForSampleNum(sampleNumAtChunkStart + i);
-			}
-			CurrentSampleFileOffset = (int64) chunkOffset;
+			// Are we in the file's base samples or the fragments (if any)?
+			if (!FragmentVars.bInFragments)
+			{
+				// We are in the samples from the moov box.
+				int64 DTS = sttsIt.GetCurrentTime() + Convs.BaseMediaDecodeTime;
+				uint32 Duration = sttsIt.GetCurrentDuration();
+				int64 CompositionTimeOffset = cttsIt.GetCurrentOffset();
+				int64 PTS = DTS + CompositionTimeOffset;
+
+				const uint32 Timescale = Convs.DurationFromMdhdBox.GetDenominator();
+				CurrentDTS.SetFromND(DTS, Timescale);
+				CurrentPTS.SetFromND(PTS, Timescale);
+				CurrentEffectiveDTS.SetFromND(DTS - Convs.DTSShiftAtZeroPoint, Timescale);
+				CurrentEffectivePTS.SetFromND(PTS - Convs.CompositionTimeAtZeroPoint, Timescale);
+				// Set the duration as the fraction of the duration and the timescale.
+				CurrentDuration.SetFromND(Duration, Timescale);
+				// Also set the duration as the delta of the DTS of this sample and the next
+				// in timespan units. This is to avoid transformation issues from media local
+				// time into the timescale used in engine.
+				CurrentDurationTS = FTimeFraction(DTS + Duration, Timescale).GetAsTimespan() - CurrentDTS.GetAsTimespan();
+				CurrentSampleSize = stszIt.GetCurrentSampleSize();
+				bCurrentIsSyncOrRAP = stssIt.IsSyncSample() || rapIt.GetCurrentGroupDescriptionIndex() != 0;
+
+				// Which chunk is this sample in?
+				uint32 chunkIndex = stscIt.GetCurrentChunkIndex();
+				check(chunkIndex);
+				uint64 chunkOffset = stcoIt.GetOffsetForChunkIndex(chunkIndex - 1);
+				check(chunkOffset);
+				// Which sample position within the current chunk run are we at?
+				uint32 samplePosInChunk = stscIt.GetSampleIndexInCurrentChunk();
+				// Giving us which sample number at the start of the chunk?
+				uint32 sampleNumAtChunkStart = SampleNumber - samplePosInChunk;
+				for(uint32 i=0; i<samplePosInChunk; ++i)
+				{
+					chunkOffset += stszIt.GetSampleSizeForSampleNum(sampleNumAtChunkStart + i);
+				}
+				CurrentSampleFileOffset = (int64) chunkOffset;
+			}
+			else
+			{
+				// We are in the fragment samples.
+				const int64 BaseMediaDecodeTime = 0;
+				int64 DTS = FragmentVars.TrunIterator.GetDTS() + BaseMediaDecodeTime;
+				int64 PTS = FragmentVars.TrunIterator.GetPTS() + BaseMediaDecodeTime;
+				int64 Duration = FragmentVars.TrunIterator.GetDuration();
+				const uint32 Timescale = Convs.DurationFromMdhdBox.GetDenominator();
+				CurrentDTS.SetFromND(DTS, Timescale);
+				CurrentPTS.SetFromND(PTS, Timescale);
+				CurrentEffectiveDTS.SetFromND(DTS - Convs.DTSShiftAtZeroPoint, Timescale);
+				CurrentEffectivePTS.SetFromND(PTS - Convs.CompositionTimeAtZeroPoint, Timescale);
+				// Set the duration as the fraction of the duration and the timescale.
+				CurrentDuration.SetFromND(Duration, Timescale);
+				// Also set the duration as the delta of the DTS of this sample and the next
+				// in timespan units. This is to avoid transformation issues from media local
+				// time into the timescale used in engine.
+				CurrentDurationTS = FTimeFraction(DTS + Duration, Timescale).GetAsTimespan() - CurrentDTS.GetAsTimespan();
+				CurrentSampleSize = FragmentVars.TrunIterator.GetSampleSize();
+				bCurrentIsSyncOrRAP = FragmentVars.TrunIterator.IsSyncSample();
+				CurrentSampleFileOffset = FragmentVars.TrunIterator.GetAbsoluteOffset();
+				SampleNumber = FragmentVars.TrunIterator.GetAbsoluteSampleNumber();
+			}
 		}
 
 
 		bool FMP4Track::FIterator::Next()
 		{
-			if (SampleNumber+1 < Convs.NumTotalSamples)
+			const uint32 NextSampleNumber = SampleNumber + 1;
+			if (NextSampleNumber >= Convs.NumTotalSamples)
+			{
+				return false;
+			}
+			if (NextSampleNumber < Convs.NumMoovSamples)
 			{
 				// Note: stcoIt is not an iterator, so there's nothing to call on it.
 				verify(stszIt.Next());
@@ -564,15 +789,35 @@ namespace Electra
 				verify(stscIt.Next());
 				verify(stssIt.Next());
 				verify(rapIt.Next());
-				++SampleNumber;
-				Update();
-				return true;
 			}
-			return false;
+			else
+			{
+				// Already in fragments?
+				if (FragmentVars.bInFragments)
+				{
+					if (!FragmentVars.TrunIterator.Next())
+					{
+						check(!"why did we get here?");
+						return false;
+					}
+				}
+				else
+				{
+					StartAt(NextSampleNumber);
+				}
+			}
+			SampleNumber = NextSampleNumber;
+			Update();
+			return true;
 		}
 		bool FMP4Track::FIterator::Prev()
 		{
-			if (SampleNumber > 0)
+			if (SampleNumber == 0)
+			{
+				return false;
+			}
+			const uint32 NextSampleNumber = SampleNumber - 1;
+			if (!FragmentVars.bInFragments)
 			{
 				// Note: stcoIt is not an iterator, so there's nothing to call on it.
 				verify(stszIt.Prev());
@@ -581,11 +826,18 @@ namespace Electra
 				verify(stscIt.Prev());
 				verify(stssIt.Prev());
 				verify(rapIt.Prev());
-				--SampleNumber;
-				Update();
-				return true;
 			}
-			return false;
+			else
+			{
+				// Going back from fragments into the moov?
+				if (!FragmentVars.TrunIterator.Prev())
+				{
+					StartAt(NextSampleNumber);
+				}
+			}
+			SampleNumber = NextSampleNumber;
+			Update();
+			return true;
 		}
 
 
@@ -594,16 +846,7 @@ namespace Electra
 			// The last sample number is inclusive, that is, that sample is needed.
 			if (SampleNumber+1 <= Convs.LastSample.LastSampleNumber)
 			{
-				// Note: stcoIt is not an iterator, so there's nothing to call on it.
-				verify(stszIt.Next());
-				verify(sttsIt.Next());
-				verify(cttsIt.Next());
-				verify(stscIt.Next());
-				verify(stssIt.Next());
-				verify(rapIt.Next());
-				++SampleNumber;
-				Update();
-				return true;
+				return Next();
 			}
 			return false;
 		}
@@ -611,20 +854,283 @@ namespace Electra
 		{
 			if (SampleNumber > Convs.FirstSample.SyncSampleNumber)
 			{
-				// Note: stcoIt is not an iterator, so there's nothing to call on it.
-				verify(stszIt.Prev());
-				verify(sttsIt.Prev());
-				verify(cttsIt.Prev());
-				verify(stscIt.Prev());
-				verify(stssIt.Prev());
-				verify(rapIt.Prev());
+				return Prev();
+			}
+			return false;
+		}
+
+
+
+		void FMP4Track::FIterator::FTrunIterator::SetFragmentInfo(const TSharedPtr<FFragmentInfo, ESPMode::ThreadSafe>& InFragmentInfo)
+		{
+			FragmentInfo = InFragmentInfo;
+			if (FragmentInfo && FragmentInfo->TrexBox)
+			{
+				TrexDefaultSampleDescriptionIndex = FragmentInfo->TrexBox->GetDefaultSampleDescriptionIndex();
+				TrexDefaultSampleDuration = FragmentInfo->TrexBox->GetDefaultSampleDuration();
+				TrexDefaultSampleSize = FragmentInfo->TrexBox->GetDefaultSampleSize();
+				TrexDefaultSampleFlags = FragmentInfo->TrexBox->GetDefaultSampleFlags();
+			}
+			else
+			{
+				TrexDefaultSampleDescriptionIndex = 0;
+				TrexDefaultSampleDuration = 0;
+				TrexDefaultSampleSize = 0;
+				TrexDefaultSampleFlags = 0;
+			}
+		}
+
+		bool FMP4Track::FIterator::FTrunIterator::StartAt(uint32 InSampleNum)
+		{
+			check(IsSetup());
+
+			TrunNumber = 0;
+			TSharedPtr<FFragmentInfo::FTrackTrafs, ESPMode::ThreadSafe> Traf;
+			TSharedPtr<FMP4BoxTRUN, ESPMode::ThreadSafe> Trun;
+			// In which fragment is the frame?
+			for(int32 trfIdx=0,trfIdxMax=FragmentInfo->TrackFragments.Num(); !Trun && trfIdx<trfIdxMax; ++trfIdx)
+			{
+				Traf = FragmentInfo->TrackFragments[trfIdx];
+				if (InSampleNum >= Traf->AbsoluteFirstSampleNumber &&
+					InSampleNum < Traf->AbsoluteFirstSampleNumber + Traf->NumSamples)
+				{
+					FirstAbsoluteSampleNumber = Traf->AbsoluteFirstSampleNumber;
+					TrunStartDTS = Traf->FirstTrackTotalDTS;
+					uint32 FrameNumInTraf = InSampleNum - Traf->AbsoluteFirstSampleNumber;
+					for(int32 i=0, iMax=Traf->TrunBoxes.Num(); i<iMax; ++i)
+					{
+						Trun = Traf->TrunBoxes[i];
+						if (FrameNumInTraf < Traf->TrunBoxes[i]->GetNumberOfSamples())
+						{
+							TrafNumber = trfIdx;
+							TrunNumber = i;
+							SampleNumber = FrameNumInTraf;
+							break;
+						}
+						if (Trun->HasSampleDurations())
+						{
+							TrunStartDTS += Trun->GetTotalDuration();
+						}
+						else
+						{
+							TrunStartDTS += Trun->GetNumberOfSamples() * (Traf->TfhdBox->HasDefaultSampleDuration() ? Traf->TfhdBox->GetDefaultSampleDuration() : TrexDefaultSampleDuration);
+						}
+						FirstAbsoluteSampleNumber += Trun->GetNumberOfSamples();
+						FrameNumInTraf -= Trun->GetNumberOfSamples();
+					}
+				}
+			}
+			check(Trun.IsValid());
+			if (!Trun)
+			{
+				return false;
+			}
+
+			Update(Traf);
+			Update(Trun);
+			return true;
+		}
+
+		void FMP4Track::FIterator::FTrunIterator::Update(const TSharedPtr<FFragmentInfo::FTrackTrafs, ESPMode::ThreadSafe>& InTraf)
+		{
+			DefaultSampleDescriptionIndex = InTraf->TfhdBox->HasSampleDescriptionIndex() ? InTraf->TfhdBox->GetSampleDescriptionIndex() : TrexDefaultSampleDescriptionIndex;
+			DefaultSampleDuration = InTraf->TfhdBox->HasDefaultSampleDuration() ? InTraf->TfhdBox->GetDefaultSampleDuration() : TrexDefaultSampleDuration;
+			DefaultSampleSize = InTraf->TfhdBox->HasDefaultSampleSize() ? InTraf->TfhdBox->GetDefaultSampleSize() : TrexDefaultSampleSize;
+			DefaultSampleFlags = InTraf->TfhdBox->HasDefaultSampleFlags() ? InTraf->TfhdBox->GetDefaultSampleFlags() : TrexDefaultSampleFlags;
+
+			BaseOffset = 0;
+			// Note: we do not consider any of the DataEntryImdaBox, DataEntrySeqNumImdaBox or IdentifiedMediaDataBox here.
+			if (InTraf->TfhdBox->HasBaseDataOffset())
+			{
+				BaseOffset = InTraf->TfhdBox->GetBaseDataOffset();
+			}
+			else if (InTraf->TfhdBox->IsMoofDefaultBase())
+			{
+				BaseOffset = InTraf->MoofBox->GetBoxFileOffset();
+			}
+			else
+			{
+				BaseOffset = InTraf->MoofBox->GetBoxFileOffset();
+			}
+		}
+		
+		void FMP4Track::FIterator::FTrunIterator::Update(const TSharedPtr<FMP4BoxTRUN, ESPMode::ThreadSafe>& InTrun)
+		{
+			// Determine the offset to the sample by either adding the preceding sample sizes
+			// or directly if fixed sample size.
+			DataOffset = InTrun->HasSampleOffset() ? InTrun->GetSampleOffset() : 0;
+
+			if (InTrun->HasSampleSizes())
+			{
+				const uint32* SampleSizes = InTrun->GetSampleSizes().GetData();
+				SampleSize = SampleSizes[SampleNumber];
+				for(int32 i=0; i<SampleNumber; ++i)
+				{
+					DataOffset += *SampleSizes++;
+				}
+			}
+			else
+			{
+				SampleSize = DefaultSampleSize;
+				DataOffset += SampleSize * SampleNumber;
+			}
+			AbsoluteSampleByteOffset = BaseOffset + DataOffset;
+
+			// Sample flags
+			SampleFlags = InTrun->HasSampleFlags() ? InTrun->GetSampleFlags()[SampleNumber] : (SampleNumber == 0 && InTrun->HasFirstSampleFlags()) ? InTrun->GetFirstSampleFlags() : DefaultSampleFlags;
+
+			// Duration and DTS
+			if (InTrun->HasSampleDurations())
+			{
+				DTS = 0;
+				const uint32* SampleDurations = InTrun->GetSampleDurations().GetData();
+				Duration = SampleDurations[SampleNumber];
+				for(int32 i=0; i<SampleNumber; ++i)
+				{
+					DTS += *SampleDurations++;
+				}
+			}
+			else
+			{
+				Duration = DefaultSampleDuration;
+				DTS = Duration * SampleNumber;
+			}
+			DTS += TrunStartDTS;
+			// PTS
+			if (InTrun->HasSampleCompositionTimeOffsets())
+			{
+				PTS = DTS + InTrun->GetSampleCompositionTimeOffsets()[SampleNumber];
+			}
+			else
+			{
+				PTS = DTS;
+			}
+		}
+
+		bool FMP4Track::FIterator::FTrunIterator::Prev()
+		{
+			check(IsSetup());
+			// If we are at the start of the fragments there is no previous sample to go back to.
+			if (TrafNumber <= 0 && TrunNumber <= 0 && SampleNumber <= 0)
+			{
+				return false;
+			}
+
+			TSharedPtr<FFragmentInfo::FTrackTrafs, ESPMode::ThreadSafe> Traf = FragmentInfo->TrackFragments[TrafNumber];
+			TSharedPtr<FMP4BoxTRUN, ESPMode::ThreadSafe> Trun = Traf->TrunBoxes[TrunNumber];
+			// Do we have a preceding sample in this trun?
+			if (SampleNumber > 0)
+			{
 				--SampleNumber;
-				Update();
+				SampleSize = Trun->HasSampleSizes() ? Trun->GetSampleSizes()[SampleNumber] : DefaultSampleSize;
+				DataOffset -= SampleSize;
+				AbsoluteSampleByteOffset -= SampleSize;
+				SampleFlags = Trun->HasSampleFlags() ? Trun->GetSampleFlags()[SampleNumber] : (SampleNumber == 0 && Trun->HasFirstSampleFlags()) ? Trun->GetFirstSampleFlags() : DefaultSampleFlags;
+				Duration = Trun->HasSampleDurations() ? Trun->GetSampleDurations()[SampleNumber] : DefaultSampleDuration;
+				DTS = DTS - Duration;
+				PTS = DTS + (Trun->HasSampleCompositionTimeOffsets() ? Trun->GetSampleCompositionTimeOffsets()[SampleNumber] : 0);
+				return true;
+			}
+			// Is there a trun before we can move into?
+			if (TrunNumber > 0)
+			{
+				--TrunNumber;
+				Trun = Traf->TrunBoxes[TrunNumber];
+				// Adjust the trun base values down.
+				const uint32 NumSamplesInTrun = Trun->GetNumberOfSamples();
+				check(NumSamplesInTrun);
+				FirstAbsoluteSampleNumber -= NumSamplesInTrun;
+				int64 TrunDuration = Trun->HasSampleDurations() ? Trun->GetTotalDuration() : NumSamplesInTrun * DefaultSampleDuration;
+				TrunStartDTS -= TrunDuration;
+				SampleNumber = NumSamplesInTrun - 1;
+				Update(Trun);
+				return true;
+			}
+			// Is there an earlier traf to move into?
+			if (TrafNumber > 0)
+			{
+				--TrafNumber;
+				Traf = FragmentInfo->TrackFragments[TrafNumber];
+				Update(Traf);
+				FirstAbsoluteSampleNumber = Traf->AbsoluteFirstSampleNumber;
+				TrunStartDTS = Traf->FirstTrackTotalDTS;
+				// There should be trun boxes here.
+				check(!Traf->TrunBoxes.IsEmpty());
+				TrunNumber = Traf->TrunBoxes.Num() - 1;
+				for(int32 i=0; i<TrunNumber; ++i)
+				{
+					Trun = Traf->TrunBoxes[i];
+					FirstAbsoluteSampleNumber += Trun->GetNumberOfSamples();
+					if (Trun->HasSampleDurations())
+					{
+						TrunStartDTS += Trun->GetTotalDuration();
+					}
+					else
+					{
+						TrunStartDTS += Trun->GetNumberOfSamples() * (Traf->TfhdBox->HasDefaultSampleDuration() ? Traf->TfhdBox->GetDefaultSampleDuration() : TrexDefaultSampleDuration);
+					}
+				}
+				Trun = Traf->TrunBoxes[TrunNumber];
+				// There should be samples in here.
+				check(Trun->GetNumberOfSamples());
+				SampleNumber = Trun->GetNumberOfSamples() - 1;
+				Update(Trun);
+				return true;
+			}
+			return false;
+		}
+
+		bool FMP4Track::FIterator::FTrunIterator::Next()
+		{
+			check(IsSetup());
+			check(TrafNumber >= 0 && TrunNumber >= 0 && SampleNumber >= 0);
+			TSharedPtr<FFragmentInfo::FTrackTrafs, ESPMode::ThreadSafe> Traf = FragmentInfo->TrackFragments[TrafNumber];
+			TSharedPtr<FMP4BoxTRUN, ESPMode::ThreadSafe> Trun = Traf->TrunBoxes[TrunNumber];
+			
+			// Still samples in this trun?
+			if (SampleNumber + 1 < (int32)Trun->GetNumberOfSamples())
+			{
+				++SampleNumber;
+				DataOffset += SampleSize;
+				AbsoluteSampleByteOffset += SampleSize;
+				DTS += Duration;
+				PTS = DTS + (Trun->HasSampleCompositionTimeOffsets() ? Trun->GetSampleCompositionTimeOffsets()[SampleNumber] : 0);
+				SampleSize = Trun->HasSampleSizes() ? Trun->GetSampleSizes()[SampleNumber] : DefaultSampleSize;
+				SampleFlags = Trun->HasSampleFlags() ? Trun->GetSampleFlags()[SampleNumber] : (SampleNumber == 0 && Trun->HasFirstSampleFlags()) ? Trun->GetFirstSampleFlags() : DefaultSampleFlags;
+				Duration = Trun->HasSampleDurations() ? Trun->GetSampleDurations()[SampleNumber] : DefaultSampleDuration;
+				return true;
+			}
+			// Is there another trun following we can move into?
+			if (TrunNumber + 1 < Traf->TrunBoxes.Num())
+			{
+				TrunStartDTS = DTS + Duration;
+				FirstAbsoluteSampleNumber += Trun->GetNumberOfSamples();
+				++TrunNumber;
+				Trun = Traf->TrunBoxes[TrunNumber];
+				SampleNumber = 0;
+				Update(Trun);
+				return true;
+			}
+			// Is there another traf we can move into?
+			if (TrafNumber + 1 < FragmentInfo->TrackFragments.Num())
+			{
+				++TrafNumber;
+				Traf = FragmentInfo->TrackFragments[TrafNumber];
+				Update(Traf);
+				FirstAbsoluteSampleNumber = Traf->AbsoluteFirstSampleNumber;
+				TrunStartDTS = Traf->FirstTrackTotalDTS;
+				// There should be trun boxes here.
+				check(!Traf->TrunBoxes.IsEmpty());
+				TrunNumber = 0;
+				Trun = Traf->TrunBoxes[0];
+				SampleNumber = 0;
+				Update(Trun);
 				return true;
 			}
 			return false;
 		}
 
+
 	} // namespace UtilitiesMP4
 
 } // namespace Electra
diff --git a/Engine/Plugins/Media/ElectraUtil/Source/ElectraBase/Public/PlayerTime.h b/Engine/Plugins/Media/ElectraUtil/Source/ElectraBase/Public/PlayerTime.h
--- a/Engine/Plugins/Media/ElectraUtil/Source/ElectraBase/Public/PlayerTime.h
+++ b/Engine/Plugins/Media/ElectraUtil/Source/ElectraBase/Public/PlayerTime.h
@@ -519,7 +519,14 @@ public:
 	inline FTimeValue operator / (int32 Scale) const
 	{
 		FTimeValue Result(*this);
-		Result.HNS /= Scale;
+		if (Scale)
+		{
+			Result.HNS /= Scale;
+		}
+		else
+		{
+			Result.SetToPositiveInfinity();
+		}
 		return Result;
 	}
 
@@ -676,6 +683,20 @@ public:
 		return *this;
 	}
 
+	FTimeFraction& SetNumerator(int64 InNumerator)
+	{
+		Numerator = InNumerator;
+		bIsValid = true;
+		return *this;
+	}
+
+	FTimeFraction& SetDenominator(uint32 InDenominator)
+	{
+		Denominator = InDenominator;
+		bIsValid = true;
+		return *this;
+	}
+
 	FTimeFraction& SetToPositiveInfinity()
 	{
 		Numerator = 0;
diff --git a/Engine/Plugins/Media/ElectraUtil/Source/ElectraBase/Public/Utilities/MP4Boxes/MP4BoxBase.h b/Engine/Plugins/Media/ElectraUtil/Source/ElectraBase/Public/Utilities/MP4Boxes/MP4BoxBase.h
--- a/Engine/Plugins/Media/ElectraUtil/Source/ElectraBase/Public/Utilities/MP4Boxes/MP4BoxBase.h
+++ b/Engine/Plugins/Media/ElectraUtil/Source/ElectraBase/Public/Utilities/MP4Boxes/MP4BoxBase.h
@@ -37,19 +37,24 @@ namespace Electra
 			//-------------------------------------------------------------------------
 			// Parent box query methods
 			//
+			TSharedPtr<FMP4BoxBase, ESPMode::ThreadSafe> GetParentBox() const
+			{
+				return ParentBox.Pin();
+			}
+
 			template<typename T>
 			TSharedPtr<T, ESPMode::ThreadSafe> FindParentBox(uint32 InType) const
 			{
-				TSharedPtr<T, ESPMode::ThreadSafe> Parent = StaticCastSharedPtr<T>(ParentBox.Pin());
+				TSharedPtr<FMP4BoxBase, ESPMode::ThreadSafe> Parent = ParentBox.Pin();
 				while(Parent.IsValid())
 				{
 					if (Parent->GetType() == InType)
 					{
-						return Parent;
+						return StaticCastSharedPtr<T>(Parent);
 					}
-					Parent = StaticCastSharedPtr<T>(Parent->ParentBox.Pin());
+					Parent = Parent->ParentBox.Pin();
 				}
-				return Parent;
+				return nullptr;
 			}
 
 			//-------------------------------------------------------------------------
@@ -85,6 +90,36 @@ namespace Electra
 				return nullptr;
 			}
 
+			template<typename T, typename Predicate>
+			TSharedPtr<T, ESPMode::ThreadSafe> FindBoxRecursiveByPredicate(uint32 InType, Predicate InPredicate, int32 InMaxDepth=8) const
+			{
+				// First pass, check all children if they are of the type
+				for(int32 i=0,iMax=Children.Num(); i<iMax; ++i)
+				{
+					if (Children[i]->GetType() == InType)
+					{
+						TSharedPtr<T, ESPMode::ThreadSafe> Elem = StaticCastSharedPtr<T>(Children[i]);
+						if (::Invoke(InPredicate, Elem))
+						{
+							return Elem;
+						}
+					}
+				}
+				// Second pass, check children recursively
+				if (InMaxDepth > 0)
+				{
+					for(int32 i=0, iMax=Children.Num(); i<iMax; ++i)
+					{
+						TSharedPtr<T, ESPMode::ThreadSafe> Box = Children[i]->FindBoxRecursiveByPredicate<T, Predicate>(InType, InPredicate, InMaxDepth - 1);
+						if (Box.IsValid())
+						{
+							return Box;
+						}
+					}
+				}
+				return nullptr;
+			}
+
 
 			// Returns all instances of a given box type from the direct children of THIS BOX ONLY.
 			template<typename T>
@@ -99,6 +134,29 @@ namespace Electra
 				}
 			}
 
+			template<typename T, typename Predicate>
+			void GetAllBoxInstancesByPredicate(TArray<TSharedPtr<T, ESPMode::ThreadSafe>>& OutAllBoxes, uint32 InType, Predicate InPredicate, int32 InMaxDepth=8) const
+			{
+				for(int32 i=0,iMax=Children.Num(); i<iMax; ++i)
+				{
+					if (Children[i]->GetType() == InType)
+					{
+						TSharedPtr<T, ESPMode::ThreadSafe> Elem = StaticCastSharedPtr<T>(Children[i]);
+						if (::Invoke(InPredicate, Elem))
+						{
+							OutAllBoxes.Emplace(Elem);
+						}
+					}
+				}
+				if (InMaxDepth > 0)
+				{
+					for(int32 i=0, iMax=Children.Num(); i<iMax; ++i)
+					{
+						Children[i]->GetAllBoxInstancesByPredicate<T, Predicate>(OutAllBoxes, InType, InPredicate, InMaxDepth - 1);
+					}
+				}
+			}
+
 
 
 			// Called by building the box tree. Not to be called by user code.
diff --git a/Engine/Plugins/Media/ElectraUtil/Source/ElectraBase/Public/Utilities/MP4Boxes/MP4Boxes.h b/Engine/Plugins/Media/ElectraUtil/Source/ElectraBase/Public/Utilities/MP4Boxes/MP4Boxes.h
--- a/Engine/Plugins/Media/ElectraUtil/Source/ElectraBase/Public/Utilities/MP4Boxes/MP4Boxes.h
+++ b/Engine/Plugins/Media/ElectraUtil/Source/ElectraBase/Public/Utilities/MP4Boxes/MP4Boxes.h
@@ -845,7 +845,7 @@ namespace Electra
 
 
 		/**
-		 * `stsz` box
+		 * `stsz` / `stz2` box
 		 * ISO/IEC 14496-12:2022 - 8.7.3 Sample size boxes
 		 */
 		class FMP4BoxSTSZ : public FMP4BoxBase
@@ -854,6 +854,13 @@ namespace Electra
 			static TSharedPtr<FMP4BoxBase, ESPMode::ThreadSafe> Create(const TWeakPtr<FMP4BoxBase, ESPMode::ThreadSafe>& InParent, const FMP4BoxInfo& InBoxInfo)
 			{ return MakeShareable<FMP4BoxSTSZ>(new FMP4BoxSTSZ(InParent, InBoxInfo)); }
 			virtual ~FMP4BoxSTSZ() = default;
+			uint32 GetType() const override
+			{
+				// This class handles both `stsz` and `stz2`.
+				// Since `stz2` is rare and most of the time we expect `stsz` we always return
+				// this type, so user code does not need to worry about the difference.
+				return MakeBoxAtom('s','t','s','z');
+			}
 			ELECTRABASE_API uint32 GetNumberOfSamples();
 			ELECTRABASE_API uint32 GetSizeOfSample(uint32 InIndex);
 		protected:
@@ -1024,7 +1031,7 @@ namespace Electra
 			 * Returns the fragment duration in the timescale of the `mvhd` box.
 			 * This is just the value as stored in the box. You need to apply the timescale yourself.
 			 */
-			ELECTRABASE_API uint64 GetFragmentDuration();
+			ELECTRABASE_API int64 GetFragmentDuration();
 		protected:
 			FMP4BoxMEHD(const TWeakPtr<FMP4BoxBase, ESPMode::ThreadSafe>& InParent, const FMP4BoxInfo& InBoxInfo)
 				: FMP4BoxBase(InParent, InBoxInfo)
@@ -1191,6 +1198,7 @@ namespace Electra
 			ELECTRABASE_API const TArray<uint32>& GetSampleFlags();
 			ELECTRABASE_API bool HasSampleCompositionTimeOffsets();
 			ELECTRABASE_API const TArray<int32>& GetSampleCompositionTimeOffsets();
+			ELECTRABASE_API int64 GetTotalDuration();
 		protected:
 			FMP4BoxTRUN(const TWeakPtr<FMP4BoxBase, ESPMode::ThreadSafe>& InParent, const FMP4BoxInfo& InBoxInfo)
 				: FMP4BoxBase(InParent, InBoxInfo)
diff --git a/Engine/Plugins/Media/ElectraUtil/Source/ElectraBase/Public/Utilities/MP4Boxes/MP4Track.h b/Engine/Plugins/Media/ElectraUtil/Source/ElectraBase/Public/Utilities/MP4Boxes/MP4Track.h
--- a/Engine/Plugins/Media/ElectraUtil/Source/ElectraBase/Public/Utilities/MP4Boxes/MP4Track.h
+++ b/Engine/Plugins/Media/ElectraUtil/Source/ElectraBase/Public/Utilities/MP4Boxes/MP4Track.h
@@ -45,11 +45,46 @@ namespace Electra
 				int64 DTSShiftAtZeroPoint = 0;
 				uint32 TrackID = 0;
 				uint32 NumTotalSamples = 0;
+				uint32 NumMoovSamples = 0;
 			};
 
 		public:
-			static TSharedPtr<FMP4Track, ESPMode::ThreadSafe> Create(const TSharedPtr<FMP4BoxTRAK, ESPMode::ThreadSafe>& InTrakBox)
-			{ return MakeShareable<FMP4Track>(new FMP4Track(InTrakBox)); }
+			class FFragmentInfo
+			{
+			public:
+				struct FTrackTrafs
+				{
+					TSharedPtr<FMP4BoxMOOF, ESPMode::ThreadSafe> MoofBox;
+					TSharedPtr<FMP4BoxTRAF, ESPMode::ThreadSafe> TrafBox;
+					TSharedPtr<FMP4BoxTFHD, ESPMode::ThreadSafe> TfhdBox;
+					TArray<TSharedPtr<FMP4BoxTRUN, ESPMode::ThreadSafe>> TrunBoxes;
+					uint32 AbsoluteFirstSampleNumber = 0;
+					int64 FirstTrackTotalDTS = 0;
+					int64 LastTrackTotalDTSAndDur = 0;
+					int64 Duration = 0;
+					uint32 NumSamples = 0;
+				};
+				TArray<TSharedPtr<FTrackTrafs, ESPMode::ThreadSafe>> TrackFragments;
+				TSharedPtr<FMP4BoxTKHD, ESPMode::ThreadSafe> TkhdBox;
+				TSharedPtr<FMP4BoxMDHD, ESPMode::ThreadSafe> MdhdBox;
+				TSharedPtr<FMP4BoxTREX, ESPMode::ThreadSafe> TrexBox;
+				uint32 NumTotalFragmentSamples = 0;
+
+				bool HasBeenPrepared() const
+				{ return bHasBeenPrepared;}
+				ELECTRABASE_API bool Prepare(const TSharedPtr<FMP4BoxTRAK, ESPMode::ThreadSafe>& InTrakBox, const TSharedPtr<FMP4BoxMVEX, ESPMode::ThreadSafe>& InMvexBox, const TArray<TSharedPtr<Electra::UtilitiesMP4::FMP4BoxBase>>& InAllRootBoxes);
+				ELECTRABASE_API int64 GetTotalDuration();
+			private:
+				void CalculateDurations();
+				int64 BaseTrackDuration = 0;
+				int64 SumOfFragmentDurations = 0;
+				uint32 NumBaseTrackSamples = 0;
+				bool bDidCalculateDuration = false;
+				bool bHasBeenPrepared = false;
+			};
+
+			static TSharedPtr<FMP4Track, ESPMode::ThreadSafe> Create(const TSharedPtr<FMP4BoxTRAK, ESPMode::ThreadSafe>& InTrakBox, const TSharedPtr<Electra::UtilitiesMP4::FMP4Track::FFragmentInfo, ESPMode::ThreadSafe>& InFragmentInfo)
+			{ return MakeShareable<FMP4Track>(new FMP4Track(InTrakBox, InFragmentInfo)); }
 			virtual ~FMP4Track() = default;
 
 			ELECTRABASE_API bool Prepare(FTimeFraction InFullMovieDuration, FTimeFraction InAdjustedMovieDuration);
@@ -138,10 +173,71 @@ namespace Electra
 				TSharedPtr<FIterator, ESPMode::ThreadSafe> Clone()
 				{ return MakeShareable(new FIterator(*this)); }
 			protected:
+				class FTrunIterator
+				{
+				public:
+					FTrunIterator() = default;
+					void SetFragmentInfo(const TSharedPtr<FFragmentInfo, ESPMode::ThreadSafe>& InFragmentInfo);
+					bool IsSetup() const
+					{ return FragmentInfo.IsValid(); }
+
+					bool StartAt(uint32 InSampleNum);
+					bool Prev();
+					bool Next();
+
+					int64 GetStartDTS() const
+					{ return TrunStartDTS; }
+					uint32 GetAbsoluteSampleNumber() const
+					{ return FirstAbsoluteSampleNumber + (uint32)SampleNumber; }
+					int64 GetAbsoluteOffset() const
+					{ return AbsoluteSampleByteOffset; }
+					int32 GetSampleNumber() const
+					{ return SampleNumber; }
+					int64 GetSampleSize() const
+					{ return SampleSize; }
+					int64 GetDTS() const
+					{ return DTS; }
+					int64 GetPTS() const
+					{ return PTS; }
+					int64 GetDuration() const
+					{ return Duration; }
+					bool IsSyncSample() const
+					{ return (SampleFlags & 0x00010000U) == 0; }
+				private:
+					void Update(const TSharedPtr<FMP4BoxTRUN, ESPMode::ThreadSafe>& InTrun);
+					void Update(const TSharedPtr<FFragmentInfo::FTrackTrafs, ESPMode::ThreadSafe>& InTraf);
+					TSharedPtr<FFragmentInfo, ESPMode::ThreadSafe> FragmentInfo;
+					uint32 TrexDefaultSampleDescriptionIndex = 0;
+					uint32 TrexDefaultSampleDuration = 0;
+					uint32 TrexDefaultSampleSize = 0;
+					uint32 TrexDefaultSampleFlags = 0;
+
+					int64 TrunStartDTS = 0;
+					int64 BaseOffset = 0;
+					int64 DataOffset = 0;
+					uint32 FirstAbsoluteSampleNumber = 0;
+					uint32 DefaultSampleDescriptionIndex = 0;
+					uint32 DefaultSampleDuration = 0;
+					uint32 DefaultSampleSize = 0;
+					uint32 DefaultSampleFlags = 0;
+					int32 TrafNumber = -1;
+					int32 TrunNumber = -1;
+					int32 SampleNumber = -1;
+
+					uint32 SampleFlags = 0;
+					int64 AbsoluteSampleByteOffset = 0;
+					int64 SampleSize = 0;
+					int64 DTS = 0;
+					int64 PTS = 0;
+					int64 Duration = 0;
+				};
+
 				FIterator() = default;
 				FIterator(const FIterator& InOther) = default;
+				FIterator& operator = (const FIterator&) = delete;
 				friend class FMP4Track;
 
+				void StartAt(uint32 InFrameNumber);
 				void Update();
 
 				TWeakPtr<FMP4Track, ESPMode::ThreadSafe> Track;
@@ -155,6 +251,13 @@ namespace Electra
 				FSBGPBoxIterator rapIt;
 				uint32 SampleNumber = 0;
 
+				struct FFragmentVars
+				{
+					FTrunIterator TrunIterator;
+					bool bInFragments = false;
+				};
+				FFragmentVars FragmentVars;
+
 				FTimeFraction CurrentDTS;
 				FTimeFraction CurrentPTS;
 				FTimeFraction CurrentEffectiveDTS;
@@ -191,10 +294,13 @@ namespace Electra
 			{ return Convs.FullMovieDuration; }
 
 		protected:
-			FMP4Track(const TSharedPtr<FMP4BoxTRAK, ESPMode::ThreadSafe>& InTrakBox) : TrakBox(InTrakBox)
+			FMP4Track(const TSharedPtr<FMP4BoxTRAK, ESPMode::ThreadSafe>& InTrakBox, const TSharedPtr<Electra::UtilitiesMP4::FMP4Track::FFragmentInfo, ESPMode::ThreadSafe>& InFragmentInfo)
+				: TrakBox(InTrakBox)
+				, FragmentInfo(InFragmentInfo)
 			{ }
 
 			TSharedPtr<FMP4BoxTRAK, ESPMode::ThreadSafe> TrakBox;
+			TSharedPtr<Electra::UtilitiesMP4::FMP4Track::FFragmentInfo, ESPMode::ThreadSafe> FragmentInfo;
 
 			TSharedPtr<FMP4BoxTKHD, ESPMode::ThreadSafe> TkhdBox;
 			TSharedPtr<FMP4BoxELST, ESPMode::ThreadSafe> ElstBox;
