ElectraUtil: Added parsing of the mp4's sub-sample box

#rnx

#Synced-CL 49644370

diff --git a/Engine/Plugins/Media/ElectraUtil/Source/ElectraBase/Private/Utilities/MP4Boxes/MP4Boxes.cpp b/Engine/Plugins/Media/ElectraUtil/Source/ElectraBase/Private/Utilities/MP4Boxes/MP4Boxes.cpp
--- a/Engine/Plugins/Media/ElectraUtil/Source/ElectraBase/Private/Utilities/MP4Boxes/MP4Boxes.cpp
+++ b/Engine/Plugins/Media/ElectraUtil/Source/ElectraBase/Private/Utilities/MP4Boxes/MP4Boxes.cpp
@@ -11,6 +11,17 @@ namespace Electra
 	namespace UtilitiesMP4
 	{
 
+		TArray<uint8> FMP4BoxBase::GetBoxDataRAW() const
+		{
+			TArray<uint8> Raw;
+			Raw.AddUninitialized(8 + BoxInfo.Data.Num());
+			uint32* bd = reinterpret_cast<uint32*>(Raw.GetData());
+			*bd++ = EndianSwap(static_cast<uint32>(Raw.Num()));
+			*bd++ = EndianSwap(BoxInfo.Type);
+			FMemory::Memcpy(bd, BoxInfo.Data.GetData(), BoxInfo.Data.Num());
+			return Raw;
+		}
+
 		void FMP4BoxBase::ProcessBoxChildrenRecursively(FMP4AtomReader& InOutReader, const FMP4BoxInfo& InCurrentBoxInfo)
 		{
 			FMP4BoxInfo bi;
@@ -744,6 +755,13 @@ namespace Electra
 
 			// There can now be additional boxes following, like `esds`, `chnl`, `dmix`, `btrt`, etc.
 			ProcessBoxChildrenRecursively(ar, BoxInfo);
+
+			// Check if there is an `srat` box to get the 32 bit sample rate from.
+			auto SRATBox = FindBoxRecursive<FMP4BoxSRAT>(Electra::UtilitiesMP4::MakeBoxAtom('s','r','a','t'), 0);
+			if (SRATBox.IsValid())
+			{
+				Parsed->SampleRate = SRATBox->GetSampleRate();
+			}
 		}
 
 		uint16 FMP4BoxAudioSampleEntry::GetDataReferenceIndex()
@@ -1224,6 +1242,32 @@ namespace Electra
 			return Parsed->AverageBitrate;
 		}
 
+/****************************************************************************************************************************************************/
+
+		struct FMP4BoxSRAT::FParsed
+		{
+			uint32 SampleRate = 0;
+			uint32 Flags = 0;
+			uint8 Version = 0;
+		};
+		void FMP4BoxSRAT::ParseIfRequired()
+		{
+			if (Parsed.IsValid())
+			{
+				return;
+			}
+			Parsed = MakePimpl<FParsed>();
+			FMP4AtomReader ar(BoxInfo.Data);
+			ar.ReadVersionAndFlags(Parsed->Version, Parsed->Flags);
+			ar.Read(Parsed->SampleRate);
+		}
+
+		uint32 FMP4BoxSRAT::GetSampleRate()
+		{
+			ParseIfRequired();
+			return Parsed->SampleRate;
+		}
+
 /****************************************************************************************************************************************************/
 
 		struct FMP4BoxSTTS::FParsed
@@ -1508,6 +1552,77 @@ namespace Electra
 			return InIndex < Parsed->NumberOfEntries ? Parsed->SizeArray32 ? GetFromBigEndian(Parsed->SizeArray32[InIndex]) : GetFromBigEndian(Parsed->SizeArray64[InIndex]) : 0;
 		}
 
+/****************************************************************************************************************************************************/
+
+		struct FMP4BoxSUBS::FParsed
+		{
+			TArray<FMP4BoxSUBS::FEntry> Entries;
+			uint32 Flags = 0;
+			uint8 Version = 0;
+		};
+		void FMP4BoxSUBS::ParseIfRequired()
+		{
+			if (Parsed.IsValid())
+			{
+				return;
+			}
+			Parsed = MakePimpl<FParsed>();
+			FMP4AtomReader ar(BoxInfo.Data);
+			ar.ReadVersionAndFlags(Parsed->Version, Parsed->Flags);
+			uint32 entry_count;
+			ar.Read(entry_count);
+			Parsed->Entries.SetNum(entry_count);
+			uint32 SampleIndex = 0;
+			for(uint32 i=0; i<entry_count; ++i)
+			{
+				uint32 sample_delta;
+				uint16 subsample_count;
+				ar.Read(sample_delta);
+				ar.Read(subsample_count);
+				SampleIndex += sample_delta;
+				check(i || sample_delta);	// the first delta must not be zero
+				Parsed->Entries[i].FirstSampleNumber = SampleIndex;
+				if (subsample_count)
+				{
+					uint32 subsample_offset = 0;
+					Parsed->Entries[i].SubSamples.Reserve(subsample_count);
+					for(int32 j=0; j<subsample_count; ++j)
+					{
+						FMP4BoxSUBS::FSubsample ss;
+						if (Parsed->Version == 1)
+						{
+							ar.Read(ss.subsample_size);
+						}
+						else
+						{
+							uint16 sssize;
+							ar.Read(sssize);
+							ss.subsample_size = sssize;
+						}
+						ar.Read(ss.subsample_priority);
+						ar.Read(ss.discardable);
+						ar.Read(ss.codec_specific_parameters);
+						// Calculate the offset into the sample for convenience.
+						ss.subsample_offset = subsample_offset;
+						subsample_offset += ss.subsample_size;
+						Parsed->Entries[i].SubSamples.Emplace(MoveTemp(ss));
+					}
+				}
+			}
+		}
+
+		uint32 FMP4BoxSUBS::GetFlags()
+		{
+			ParseIfRequired();
+			return Parsed->Flags;
+		}
+
+		const TArray<FMP4BoxSUBS::FEntry>& FMP4BoxSUBS::GetEntries()
+		{
+			ParseIfRequired();
+			return Parsed->Entries;
+		}
+
 /****************************************************************************************************************************************************/
 
 		struct FMP4BoxSAIZ::FParsed
@@ -2586,6 +2701,7 @@ namespace Electra
 			FactoryMap.Add(MakeBoxAtom('s','t','z','2'), &FMP4BoxSTSZ::Create);
 			FactoryMap.Add(MakeBoxAtom('s','t','c','o'), &FMP4BoxSTCO::Create);
 			FactoryMap.Add(MakeBoxAtom('c','o','6','4'), &FMP4BoxSTCO::Create);
+			FactoryMap.Add(MakeBoxAtom('s','u','b','s'), &FMP4BoxSUBS::Create);
 			FactoryMap.Add(MakeBoxAtom('s','a','i','z'), &FMP4BoxSAIZ::Create);
 			FactoryMap.Add(MakeBoxAtom('s','a','i','o'), &FMP4BoxSAIO::Create);
 			FactoryMap.Add(MakeBoxAtom('s','g','p','d'), &FMP4BoxSGPD::Create);
@@ -2624,6 +2740,7 @@ namespace Electra
 			FactoryMap.Add(MakeBoxAtom('m','d','c','v'), &FMP4BoxMDCV::Create);
 			FactoryMap.Add(MakeBoxAtom('b','t','r','t'), &FMP4BoxBTRT::Create);
 			FactoryMap.Add(MakeBoxAtom('p','a','s','p'), &FMP4BoxPASP::Create);
+			FactoryMap.Add(MakeBoxAtom('s','r','a','t'), &FMP4BoxSRAT::Create);
 
 			FactoryMap.Add(MakeBoxAtom('s','i','n','f'), &FMP4BoxSINF::Create);
 			FactoryMap.Add(MakeBoxAtom('f','r','m','a'), &FMP4BoxFRMA::Create);
diff --git a/Engine/Plugins/Media/ElectraUtil/Source/ElectraBase/Private/Utilities/MP4Boxes/MP4Track.cpp b/Engine/Plugins/Media/ElectraUtil/Source/ElectraBase/Private/Utilities/MP4Boxes/MP4Track.cpp
--- a/Engine/Plugins/Media/ElectraUtil/Source/ElectraBase/Private/Utilities/MP4Boxes/MP4Track.cpp
+++ b/Engine/Plugins/Media/ElectraUtil/Source/ElectraBase/Private/Utilities/MP4Boxes/MP4Track.cpp
@@ -66,6 +66,7 @@ namespace Electra
 					Frg->MoofBox = Frg->TrafBox->FindParentBox<FMP4BoxMOOF>(MakeBoxAtom('m','o','o','f'));
 					Frg->TfhdBox = Frg->TrafBox->FindBoxRecursive<FMP4BoxTFHD>(MakeBoxAtom('t','f','h','d'), 0);
 					Frg->TrafBox->GetAllBoxInstances<FMP4BoxTRUN>(Frg->TrunBoxes, MakeBoxAtom('t','r','u','n'));
+					Frg->TrafBox->GetAllBoxInstances<FMP4BoxSUBS>(Frg->SubsBoxes, MakeBoxAtom('s','u','b','s'));
 					Frg->AbsoluteFirstSampleNumber = NumBaseTrackSamples + NumTotalFragmentSamples;
 
 					Frg->NumSamples = 0;
@@ -249,6 +250,7 @@ namespace Electra
 			// Optional sample information boxes:
 			CttsBox = StblBox->FindBoxRecursive<FMP4BoxCTTS>(MakeBoxAtom('c','t','t','s'), 0);
 			StssBox = StblBox->FindBoxRecursive<FMP4BoxSTSS>(MakeBoxAtom('s','t','s','s'), 0);
+			StblBox->GetAllBoxInstances(SubsBoxes, MakeBoxAtom('s','u','b','s'));
 			StblBox->GetAllBoxInstances(SgpdBoxes, MakeBoxAtom('s','g','p','d'));
 			StblBox->GetAllBoxInstances(SbgpBoxes, MakeBoxAtom('s','b','g','p'));
 
@@ -539,6 +541,8 @@ namespace Electra
 				It->stcoIt.SetBox(StcoBox);
 				It->stscIt.SetBox(StscBox, Convs.NumMoovSamples);
 				It->stssIt.SetBox(StssBox, Convs.NumMoovSamples);
+				// Are there subsample boxes? If yes, we use the first one only in the iterator for now.
+				It->subsIt.SetBox(SubsBoxes.Num() ? SubsBoxes[0] : nullptr, Convs.NumMoovSamples);
 				// Do we have a `rap ` group?
 				TSharedPtr<FMP4BoxSGPD, ESPMode::ThreadSafe>* rapSGPD = SgpdBoxes.FindByPredicate([](const TSharedPtr<FMP4BoxSGPD, ESPMode::ThreadSafe>& e){ return e->GetGroupingType() == MakeBoxAtom('r','a','p',' ');});
 				TSharedPtr<FMP4BoxSBGP, ESPMode::ThreadSafe>* rapSBGP = SbgpBoxes.FindByPredicate([](const TSharedPtr<FMP4BoxSBGP, ESPMode::ThreadSafe>& e){ return e->GetGroupingType() == MakeBoxAtom('r','a','p',' ');});
@@ -694,6 +698,7 @@ namespace Electra
 				stscIt.SetToSampleNumber(InFrameNumber);
 				stssIt.SetToSampleNumber(InFrameNumber);
 				rapIt.SetToSampleNumber(InFrameNumber);
+				subsIt.SetToSampleNumber(InFrameNumber);
 				FragmentVars.bInFragments = false;
 			}
 			else
@@ -731,6 +736,9 @@ namespace Electra
 				CurrentDurationTS = FTimeFraction(DTS + Duration, Timescale).GetAsTimespan() - CurrentDTS.GetAsTimespan();
 				CurrentSampleSize = stszIt.GetCurrentSampleSize();
 				bCurrentIsSyncOrRAP = stssIt.IsSyncSample() || rapIt.GetCurrentGroupDescriptionIndex() != 0;
+				// Subsample list
+				CurrentSubsampleList = subsIt.GetCurrentSubsamples();
+				CurrentSubsampleBoxFlags = subsIt.GetSubsampleBoxFlags();
 
 				// Which chunk is this sample in?
 				uint32 chunkIndex = stscIt.GetCurrentChunkIndex();
@@ -769,6 +777,9 @@ namespace Electra
 				bCurrentIsSyncOrRAP = FragmentVars.TrunIterator.IsSyncSample();
 				CurrentSampleFileOffset = FragmentVars.TrunIterator.GetAbsoluteOffset();
 				SampleNumber = FragmentVars.TrunIterator.GetAbsoluteSampleNumber();
+				// Subsample list
+				CurrentSubsampleList = FragmentVars.TrunIterator.GetSubsampleList();
+				CurrentSubsampleBoxFlags = FragmentVars.TrunIterator.GetSubsampleBoxFlags();
 			}
 		}
 
@@ -789,6 +800,7 @@ namespace Electra
 				verify(stscIt.Next());
 				verify(stssIt.Next());
 				verify(rapIt.Next());
+				verify(subsIt.Next());
 			}
 			else
 			{
@@ -826,6 +838,7 @@ namespace Electra
 				verify(stscIt.Prev());
 				verify(stssIt.Prev());
 				verify(rapIt.Prev());
+				verify(subsIt.Prev());
 			}
 			else
 			{
@@ -918,6 +931,16 @@ namespace Electra
 						FirstAbsoluteSampleNumber += Trun->GetNumberOfSamples();
 						FrameNumInTraf -= Trun->GetNumberOfSamples();
 					}
+
+					// While there can be multiple `trun` boxes there can only be one `subs` box since the
+					// standard mandates that the `flags` field shall differ between multiple `subs` boxes,
+					// but because the flags is codec specific and hence has a certain meaning it would
+					// have to be the same in the `subs` boxes if there were several, which is forbidden.
+					// See ISO/IEC 14496-12:2022 Section 8.7.7 Sub-sample information box
+
+					// Are there subsample boxes? If yes, we use the first one only in the iterator for now.
+					subsIt.SetBox(Traf->SubsBoxes.Num() ? Traf->SubsBoxes[0] : nullptr, Traf->NumSamples);
+					subsIt.SetToSampleNumber(SampleNumber);
 				}
 			}
 			check(Trun.IsValid());
@@ -953,7 +976,7 @@ namespace Electra
 				BaseOffset = InTraf->MoofBox->GetBoxFileOffset();
 			}
 		}
-		
+
 		void FMP4Track::FIterator::FTrunIterator::Update(const TSharedPtr<FMP4BoxTRUN, ESPMode::ThreadSafe>& InTrun)
 		{
 			// Determine the offset to the sample by either adding the preceding sample sizes
@@ -1005,6 +1028,8 @@ namespace Electra
 			{
 				PTS = DTS;
 			}
+			// Subsample list
+			SubsampleList = subsIt.GetCurrentSubsamples();
 		}
 
 		bool FMP4Track::FIterator::FTrunIterator::Prev()
@@ -1016,6 +1041,9 @@ namespace Electra
 				return false;
 			}
 
+			// Step the subsample iterator. There is only one for all truns so we do not need to handle moving across multiple.
+			subsIt.Prev();
+
 			TSharedPtr<FFragmentInfo::FTrackTrafs, ESPMode::ThreadSafe> Traf = FragmentInfo->TrackFragments[TrafNumber];
 			TSharedPtr<FMP4BoxTRUN, ESPMode::ThreadSafe> Trun = Traf->TrunBoxes[TrunNumber];
 			// Do we have a preceding sample in this trun?
@@ -1029,6 +1057,8 @@ namespace Electra
 				Duration = Trun->HasSampleDurations() ? Trun->GetSampleDurations()[SampleNumber] : DefaultSampleDuration;
 				DTS = DTS - Duration;
 				PTS = DTS + (Trun->HasSampleCompositionTimeOffsets() ? Trun->GetSampleCompositionTimeOffsets()[SampleNumber] : 0);
+				// Update subsample list
+				SubsampleList = subsIt.GetCurrentSubsamples();
 				return true;
 			}
 			// Is there a trun before we can move into?
@@ -1074,6 +1104,11 @@ namespace Electra
 				// There should be samples in here.
 				check(Trun->GetNumberOfSamples());
 				SampleNumber = Trun->GetNumberOfSamples() - 1;
+
+				// Are there subsample boxes? If yes, we use the first one only in the iterator for now.
+				subsIt.SetBox(Traf->SubsBoxes.Num() ? Traf->SubsBoxes[0] : nullptr, Traf->NumSamples);
+				subsIt.SetToSampleNumber(SampleNumber);
+
 				Update(Trun);
 				return true;
 			}
@@ -1086,7 +1121,10 @@ namespace Electra
 			check(TrafNumber >= 0 && TrunNumber >= 0 && SampleNumber >= 0);
 			TSharedPtr<FFragmentInfo::FTrackTrafs, ESPMode::ThreadSafe> Traf = FragmentInfo->TrackFragments[TrafNumber];
 			TSharedPtr<FMP4BoxTRUN, ESPMode::ThreadSafe> Trun = Traf->TrunBoxes[TrunNumber];
-			
+
+			// Advance the subsample iterator. There is only one for all truns so we do not need to handle moving across multiple.
+			subsIt.Next();
+
 			// Still samples in this trun?
 			if (SampleNumber + 1 < (int32)Trun->GetNumberOfSamples())
 			{
@@ -1098,6 +1136,8 @@ namespace Electra
 				SampleSize = Trun->HasSampleSizes() ? Trun->GetSampleSizes()[SampleNumber] : DefaultSampleSize;
 				SampleFlags = Trun->HasSampleFlags() ? Trun->GetSampleFlags()[SampleNumber] : (SampleNumber == 0 && Trun->HasFirstSampleFlags()) ? Trun->GetFirstSampleFlags() : DefaultSampleFlags;
 				Duration = Trun->HasSampleDurations() ? Trun->GetSampleDurations()[SampleNumber] : DefaultSampleDuration;
+				// Update subsample list
+				SubsampleList = subsIt.GetCurrentSubsamples();
 				return true;
 			}
 			// Is there another trun following we can move into?
@@ -1124,6 +1164,11 @@ namespace Electra
 				TrunNumber = 0;
 				Trun = Traf->TrunBoxes[0];
 				SampleNumber = 0;
+
+				// Are there subsample boxes? If yes, we use the first one only in the iterator for now.
+				subsIt.SetBox(Traf->SubsBoxes.Num() ? Traf->SubsBoxes[0] : nullptr, Traf->NumSamples);
+				subsIt.SetToSampleNumber(SampleNumber);
+
 				Update(Trun);
 				return true;
 			}
diff --git a/Engine/Plugins/Media/ElectraUtil/Source/ElectraBase/Public/Utilities/MP4Boxes/MP4BoxBase.h b/Engine/Plugins/Media/ElectraUtil/Source/ElectraBase/Public/Utilities/MP4Boxes/MP4BoxBase.h
--- a/Engine/Plugins/Media/ElectraUtil/Source/ElectraBase/Public/Utilities/MP4Boxes/MP4BoxBase.h
+++ b/Engine/Plugins/Media/ElectraUtil/Source/ElectraBase/Public/Utilities/MP4Boxes/MP4BoxBase.h
@@ -15,6 +15,7 @@ namespace Electra
 			{ return MakeShareable<FMP4BoxBase>(new FMP4BoxBase(InParent, InBoxInfo)); }
 			virtual ~FMP4BoxBase() = default;
 
+			ELECTRABASE_API virtual TArray<uint8> GetBoxDataRAW() const;
 			virtual const FMP4BoxInfo& GetBoxInfo() const
 			{ return BoxInfo; }
 			virtual uint32 GetType() const
diff --git a/Engine/Plugins/Media/ElectraUtil/Source/ElectraBase/Public/Utilities/MP4Boxes/MP4BoxIterators.h b/Engine/Plugins/Media/ElectraUtil/Source/ElectraBase/Public/Utilities/MP4Boxes/MP4BoxIterators.h
--- a/Engine/Plugins/Media/ElectraUtil/Source/ElectraBase/Public/Utilities/MP4Boxes/MP4BoxIterators.h
+++ b/Engine/Plugins/Media/ElectraUtil/Source/ElectraBase/Public/Utilities/MP4Boxes/MP4BoxIterators.h
@@ -746,6 +746,163 @@ namespace Electra
 			uint32 CurrentEntryIndex = 0;
 		};
 
+		/****************************************************************************************************************************************************/
+
+		struct FSUBSBoxIterator
+		{
+			void SetBox(const TSharedPtr<FMP4BoxSUBS, ESPMode::ThreadSafe>& InSubsBox, uint32 InMaxSamples)
+			{
+				NumTotalSamples = InMaxSamples;
+				FirstDescribedSampleNum = 0;
+				LastDescribedSampleNum = 0;
+				CurrentEntryIndex = 0;
+				CurrentSampleNum = 0;
+				SubsBox = InSubsBox;
+				if (SubsBox.IsValid())
+				{
+					int32 NumEntries = SubsBox->GetEntries().Num();
+					if (NumEntries)
+					{
+						FirstDescribedSampleNum = SubsBox->GetEntries()[0].GetSampleNumber();
+						LastDescribedSampleNum = SubsBox->GetEntries().Last().GetSampleNumber();
+					}
+				}
+			}
+
+			void SetToSampleNumber(uint32 InSampleNum)
+			{
+				InSampleNum = InSampleNum < NumTotalSamples ? InSampleNum : NumTotalSamples-1;
+				CurrentSampleNum = InSampleNum;
+				const uint32 snumPlus1 = CurrentSampleNum + 1;
+				if (!SubsBox.IsValid())
+				{
+					return;
+				}
+				if (snumPlus1 <= FirstDescribedSampleNum)
+				{
+					CurrentEntryIndex = 0;
+				}
+				else if (snumPlus1 > LastDescribedSampleNum)
+				{
+					CurrentEntryIndex = SubsBox->GetEntries().Num() - 1;
+				}
+				else
+				{
+					// Have to locate the entry for the sample number
+					const TArray<FMP4BoxSUBS::FEntry>& Entries = SubsBox->GetEntries();
+					for(int32 i=0, iMax=Entries.Num(); i<iMax; ++i)
+					{
+						if (Entries[i].GetSampleNumber() >= snumPlus1)
+						{
+							CurrentEntryIndex = i;
+							break;
+						}
+					}
+				}
+			}
+
+			bool Next()
+			{
+				if (IsLast())
+				{
+					return false;
+				}
+				++CurrentSampleNum;
+				if (SubsBox.IsValid())
+				{
+					if (CurrentSampleNum > LastDescribedSampleNum)
+					{
+						CurrentEntryIndex = SubsBox->GetEntries().Num() - 1;
+					}
+					else
+					{
+						check(CurrentEntryIndex >= 0 && CurrentEntryIndex < SubsBox->GetEntries().Num());
+						// If the current subsample list entries sample number equals the current sample
+						// number we need to advance the list index. This is because the sample numbers
+						// in the list are 1-based while our internal sample number is 0-based.
+						// The next entry in the subsample list may not be for the next sample due to
+						// the list being sparsely coded, but that is checked for in `HaveSubsamplesForCurrent()`.
+						while(CurrentSampleNum >= SubsBox->GetEntries()[CurrentEntryIndex].GetSampleNumber())
+						{
+							if (++CurrentEntryIndex >= SubsBox->GetEntries().Num())
+							{
+								--CurrentEntryIndex;
+								break;
+							}
+						}
+					}
+				}
+				return true;
+			}
+
+			bool Prev()
+			{
+				if (IsFirst())
+				{
+					return false;
+				}
+				--CurrentSampleNum;
+				if (SubsBox.IsValid())
+				{
+					if (CurrentSampleNum < FirstDescribedSampleNum)
+					{
+						CurrentEntryIndex = 0;
+					}
+					else
+					{
+						check(CurrentEntryIndex >= 0 && CurrentEntryIndex < SubsBox->GetEntries().Num());
+						// See comments up in `Next()`
+						while(1+CurrentSampleNum < SubsBox->GetEntries()[CurrentEntryIndex].GetSampleNumber())
+						{
+							if (--CurrentEntryIndex <= 0)
+							{
+								CurrentEntryIndex = 0;
+								break;
+							}
+						}
+					}
+				}
+				return true;
+			}
+
+			TOptional<uint32> GetSubsampleBoxFlags() const
+			{
+				return SubsBox.IsValid() ? TOptional<uint32>(SubsBox->GetFlags()) : TOptional<uint32>();
+			}
+
+			bool HaveSubsamplesForCurrent() const
+			{
+				if (!FirstDescribedSampleNum || (CurrentSampleNum+1 < FirstDescribedSampleNum) || (CurrentSampleNum+1 > LastDescribedSampleNum))
+				{
+					return false;
+				}
+				check(CurrentEntryIndex < SubsBox->GetEntries().Num());
+				return CurrentSampleNum+1 == SubsBox->GetEntries()[CurrentEntryIndex].GetSampleNumber();
+			}
+
+			const FMP4BoxSUBS::FSubsampleList& GetCurrentSubsamples() const
+			{
+				return HaveSubsamplesForCurrent() ? SubsBox->GetEntries()[CurrentEntryIndex].GetEntries() : EmptySubsampleList;
+			}
+
+			uint32 GetCurrentSampleNum() const
+			{ return CurrentSampleNum; }
+			bool IsFirst() const
+			{ return CurrentSampleNum == 0; }
+			bool IsLast() const
+			{ return CurrentSampleNum + 1 >= NumTotalSamples; }
+
+		private:
+			TSharedPtr<FMP4BoxSUBS, ESPMode::ThreadSafe> SubsBox;
+			FMP4BoxSUBS::FSubsampleList EmptySubsampleList;
+
+			uint32 NumTotalSamples = 0;
+			uint32 FirstDescribedSampleNum = 0;
+			uint32 LastDescribedSampleNum = 0;
+			uint32 CurrentSampleNum = 0;
+			int32 CurrentEntryIndex = 0;
+		};
+
 
 	} // namespace UtilitiesMP4
 
diff --git a/Engine/Plugins/Media/ElectraUtil/Source/ElectraBase/Public/Utilities/MP4Boxes/MP4Boxes.h b/Engine/Plugins/Media/ElectraUtil/Source/ElectraBase/Public/Utilities/MP4Boxes/MP4Boxes.h
--- a/Engine/Plugins/Media/ElectraUtil/Source/ElectraBase/Public/Utilities/MP4Boxes/MP4Boxes.h
+++ b/Engine/Plugins/Media/ElectraUtil/Source/ElectraBase/Public/Utilities/MP4Boxes/MP4Boxes.h
@@ -448,6 +448,27 @@ namespace Electra
 		};
 
 
+		/**
+		 * `srat` box
+		 * ISO/IEC 14496-12:2022 - 12.2.3 AudioSampleEntry (12.2.3.2 Syntax)
+		 */
+		class FMP4BoxSRAT : public FMP4BoxBase
+		{
+		public:
+			static TSharedPtr<FMP4BoxBase, ESPMode::ThreadSafe> Create(const TWeakPtr<FMP4BoxBase, ESPMode::ThreadSafe>& InParent, const FMP4BoxInfo& InBoxInfo)
+			{ return MakeShareable<FMP4BoxSRAT>(new FMP4BoxSRAT(InParent, InBoxInfo)); }
+			virtual ~FMP4BoxSRAT() = default;
+			ELECTRABASE_API uint32 GetSampleRate();
+		protected:
+			FMP4BoxSRAT(const TWeakPtr<FMP4BoxBase, ESPMode::ThreadSafe>& InParent, const FMP4BoxInfo& InBoxInfo)
+				: FMP4BoxBase(InParent, InBoxInfo)
+			{ }
+			void ParseIfRequired();
+			struct FParsed;
+			TPimplPtr<FParsed> Parsed;
+		};
+
+
 		/**
 		 * Visual sample entry
 		 * ISO/IEC 14496-12:2022 - 12.1.3 Sample entry
@@ -902,6 +923,49 @@ namespace Electra
 		};
 
 
+		/**
+		 * `subs` box
+		 * ISO/IEC 14496-12:2022 - 8.7.7 Sub-sample information box
+		 */
+		class FMP4BoxSUBS : public FMP4BoxBase
+		{
+		public:
+			static TSharedPtr<FMP4BoxBase, ESPMode::ThreadSafe> Create(const TWeakPtr<FMP4BoxBase, ESPMode::ThreadSafe>& InParent, const FMP4BoxInfo& InBoxInfo)
+			{ return MakeShareable<FMP4BoxSUBS>(new FMP4BoxSUBS(InParent, InBoxInfo)); }
+			virtual ~FMP4BoxSUBS() = default;
+
+			ELECTRABASE_API uint32 GetFlags();
+
+			struct FSubsample
+			{
+				uint32 codec_specific_parameters = 0;
+				uint32 subsample_size = 0;
+				uint32 subsample_offset = 0;	// calculated for convenience.
+				uint8 subsample_priority = 0;
+				uint8 discardable = 0;
+			};
+			using FSubsampleList = TChunkedArray<FSubsample, 16>;
+			struct FEntry
+			{
+				uint32 FirstSampleNumber = 0;	// Sample number starts at 1, so 0 is not valid.
+				FSubsampleList SubSamples;
+				uint32 GetSampleNumber() const
+				{ return FirstSampleNumber; }
+				const FSubsampleList& GetEntries() const
+				{ return SubSamples; }
+			};
+			ELECTRABASE_API const TArray<FEntry>& GetEntries();
+
+		protected:
+			FMP4BoxSUBS(const TWeakPtr<FMP4BoxBase, ESPMode::ThreadSafe>& InParent, const FMP4BoxInfo& InBoxInfo)
+				: FMP4BoxBase(InParent, InBoxInfo)
+			{ }
+			void ParseIfRequired();
+			struct FParsed;
+			TPimplPtr<FParsed> Parsed;
+		};
+
+
 		/**
 		 * `saiz` box
 		 * ISO/IEC 14496-12:2022 - 8.7.8 Sample auxiliary information sizes box
diff --git a/Engine/Plugins/Media/ElectraUtil/Source/ElectraBase/Public/Utilities/MP4Boxes/MP4Track.h b/Engine/Plugins/Media/ElectraUtil/Source/ElectraBase/Public/Utilities/MP4Boxes/MP4Track.h
--- a/Engine/Plugins/Media/ElectraUtil/Source/ElectraBase/Public/Utilities/MP4Boxes/MP4Track.h
+++ b/Engine/Plugins/Media/ElectraUtil/Source/ElectraBase/Public/Utilities/MP4Boxes/MP4Track.h
@@ -58,6 +58,7 @@ namespace Electra
 					TSharedPtr<FMP4BoxTRAF, ESPMode::ThreadSafe> TrafBox;
 					TSharedPtr<FMP4BoxTFHD, ESPMode::ThreadSafe> TfhdBox;
 					TArray<TSharedPtr<FMP4BoxTRUN, ESPMode::ThreadSafe>> TrunBoxes;
+					TArray<TSharedPtr<FMP4BoxSUBS, ESPMode::ThreadSafe>> SubsBoxes;
 					uint32 AbsoluteFirstSampleNumber = 0;
 					int64 FirstTrackTotalDTS = 0;
 					int64 LastTrackTotalDTSAndDur = 0;
@@ -137,7 +138,12 @@ namespace Electra
 				// Returns the effective track's duration, as specified by an edit list. Timescale has been converted into `mdhd` timescale!
 				FTimeFraction GetEffectiveTrackDuration() const
 				{ return Convs.MappedDurationFromElstBox; }
-
+				// Returns the flags from the `subs` full box. Unset if there is no subsample box.
+				TOptional<uint32> GetSubsampleBoxFlags() const
+				{ return CurrentSubsampleBoxFlags; }
+				// Returns the list of subsamples for this sample, if there are any.
+				const FMP4BoxSUBS::FSubsampleList& GetSubsampleList() const
+				{ return CurrentSubsampleList; }
 				// Returns information about the first sample that is mapped to the 0-based timeline via `elst` box.
 				const FFirstSample& GetFirstSampleInfo() const
 				{ return Convs.FirstSample; }
@@ -203,6 +209,11 @@ namespace Electra
 					{ return Duration; }
 					bool IsSyncSample() const
 					{ return (SampleFlags & 0x00010000U) == 0; }
+
+					TOptional<uint32> GetSubsampleBoxFlags() const
+					{ return subsIt.GetSubsampleBoxFlags(); }
+					const FMP4BoxSUBS::FSubsampleList& GetSubsampleList() const
+					{ return SubsampleList; }
 				private:
 					void Update(const TSharedPtr<FMP4BoxTRUN, ESPMode::ThreadSafe>& InTrun);
 					void Update(const TSharedPtr<FFragmentInfo::FTrackTrafs, ESPMode::ThreadSafe>& InTraf);
@@ -230,6 +241,9 @@ namespace Electra
 					int64 DTS = 0;
 					int64 PTS = 0;
 					int64 Duration = 0;
+
+					FSUBSBoxIterator subsIt;
+					FMP4BoxSUBS::FSubsampleList SubsampleList;
 				};
 
 				FIterator() = default;
@@ -249,6 +263,7 @@ namespace Electra
 				FSTSSBoxIterator stssIt;
 				FSTCOBoxIterator stcoIt;
 				FSBGPBoxIterator rapIt;
+				FSUBSBoxIterator subsIt;
 				uint32 SampleNumber = 0;
 
 				struct FFragmentVars
@@ -267,6 +282,8 @@ namespace Electra
 				int64 CurrentSampleFileOffset = 0;
 				int64 CurrentSampleSize = 0;
 				bool bCurrentIsSyncOrRAP = false;
+				FMP4BoxSUBS::FSubsampleList CurrentSubsampleList;
+				TOptional<uint32> CurrentSubsampleBoxFlags;
 			};
 
 			// Create an interator starting at the first sample.
@@ -313,6 +330,7 @@ namespace Electra
 			TSharedPtr<FMP4BoxSTSS, ESPMode::ThreadSafe> StssBox;
 			TSharedPtr<FMP4BoxUDTA, ESPMode::ThreadSafe> UdtaBox;
 
+			TArray<TSharedPtr<FMP4BoxSUBS, ESPMode::ThreadSafe>> SubsBoxes;
 			TArray<TSharedPtr<FMP4BoxSGPD, ESPMode::ThreadSafe>> SgpdBoxes;
 			TArray<TSharedPtr<FMP4BoxSBGP, ESPMode::ThreadSafe>> SbgpBoxes;
 
