ElectraSamples: added destructor to release objects in a specific order

#rnx

#Synced-CL 49757097

diff --git a/Engine/Plugins/Media/ElectraUtil/Source/ElectraSamples/Private/Windows/ElectraTextureSample_Windows.cpp b/Engine/Plugins/Media/ElectraUtil/Source/ElectraSamples/Private/Windows/ElectraTextureSample_Windows.cpp
--- a/Engine/Plugins/Media/ElectraUtil/Source/ElectraSamples/Private/Windows/ElectraTextureSample_Windows.cpp
+++ b/Engine/Plugins/Media/ElectraUtil/Source/ElectraSamples/Private/Windows/ElectraTextureSample_Windows.cpp
@@ -52,6 +52,11 @@ FElectraTextureSample::~FElectraTextureSample()
 	{
 		struct FRefCountedObjs
 		{
+			~FRefCountedObjs()
+			{
+				TextureDX11.SafeRelease();
+				D3D11Device.SafeRelease();
+			}
 			TRefCountPtr<ID3D11Texture2D> TextureDX11;
 			TRefCountPtr<ID3D11Device> D3D11Device;
 		};
