Preserve subframe precision in Electra timecode construction

Use FTimecode::FromFrameTime instead of FTimecode::FromTimespan in CreateTimecodeFromMPEGDefinition so the FTimecode Subframe field is populated from the fractional frame portion of the SEI pic_timing clock timestamp.

#rnx
#Synced-CL 51630826

diff --git a/Engine/Plugins/Media/ElectraUtil/Source/ElectraSamples/Private/ElectraTextureSampleBase.cpp b/Engine/Plugins/Media/ElectraUtil/Source/ElectraSamples/Private/ElectraTextureSampleBase.cpp
--- a/Engine/Plugins/Media/ElectraUtil/Source/ElectraSamples/Private/ElectraTextureSampleBase.cpp
+++ b/Engine/Plugins/Media/ElectraUtil/Source/ElectraSamples/Private/ElectraTextureSampleBase.cpp
@@ -22,9 +22,11 @@ static TOptional<FTimecode> CreateTimecodeFromMPEGDefinition(TOptional<FFrameRat
 {
 	if (InMPEGTimecode->timing_info_present_flag)
 	{
-		const FTimespan ts(Electra::FTimeValue(InMPEGTimecode->clockTimestamp, InMPEGTimecode->time_scale).GetAsTimespan());
+		const double TotalSeconds = Electra::FTimeValue(InMPEGTimecode->clockTimestamp, InMPEGTimecode->time_scale).GetAsSeconds();
 		OutFramerate = FFrameRate(InMPEGTimecode->time_scale, InMPEGTimecode->num_units_in_tick);
-		return FTimecode::FromTimespan(ts, OutFramerate.GetValue(), InMPEGTimecode->ct_type > 1, false);
+		const FFrameTime FrameTime = OutFramerate.GetValue().AsFrameTime(TotalSeconds);
+		const bool bDropFrame = InMPEGTimecode->ct_type > 1;
+		return FTimecode::FromFrameTime(FrameTime, OutFramerate.GetValue(), bDropFrame);
 	}
 	return TOptional<FTimecode>();
 }
