[ElectraUtil] Fix null dereference in FElectraTextureSample::GetBuffer() on Windows

The Windows implementation of GetBuffer() unconditionally dereferences the Buffer TSharedPtr, causing an assertion failure when Buffer has not been populated (e.g. GPU-only decode paths). The Apple, Linux, and Android implementations all handle this case by returning nullptr.

#rnx
#Synced-CL 52001901

diff --git a/Engine/Plugins/Media/ElectraUtil/Source/ElectraSamples/Private/Windows/ElectraTextureSample_Windows.cpp b/Engine/Plugins/Media/ElectraUtil/Source/ElectraSamples/Private/Windows/ElectraTextureSample_Windows.cpp
--- a/Engine/Plugins/Media/ElectraUtil/Source/ElectraSamples/Private/Windows/ElectraTextureSample_Windows.cpp
+++ b/Engine/Plugins/Media/ElectraUtil/Source/ElectraSamples/Private/Windows/ElectraTextureSample_Windows.cpp
@@ -681,7 +681,7 @@ bool FElectraTextureSample::Convert(FRHICommandListImmediate& RHICmdList, FTextu
 
 const void* FElectraTextureSample::GetBuffer()
 {
-	return Buffer->GetData();
+	return Buffer ? Buffer->GetData() : nullptr;
 }
 
 
