[ElectraUtil] Fix timecode subframe precision loss in CreateTimecodeFromMPEGDefinition

#rnx
#Synced-CL 52001901

diff --git a/Engine/Plugins/Media/ElectraUtil/Source/ElectraSamples/Private/ElectraTextureSampleBase.cpp b/Engine/Plugins/Media/ElectraUtil/Source/ElectraSamples/Private/ElectraTextureSampleBase.cpp
--- a/Engine/Plugins/Media/ElectraUtil/Source/ElectraSamples/Private/ElectraTextureSampleBase.cpp
+++ b/Engine/Plugins/Media/ElectraUtil/Source/ElectraSamples/Private/ElectraTextureSampleBase.cpp
@@ -22,7 +22,12 @@ static TOptional<FTimecode> CreateTimecodeFromMPEGDefinition(TOptional<FFrameRat
 {
 	if (InMPEGTimecode->timing_info_present_flag)
 	{
-		const double TotalSeconds = Electra::FTimeValue(InMPEGTimecode->clockTimestamp, InMPEGTimecode->time_scale).GetAsSeconds();
+		// Compute TotalSeconds directly as a double division to avoid precision loss from
+		// FTimeValue's intermediate conversion to hundred-nanosecond integer ticks (HNS).
+		// FTimeValue::SetFromND() truncates via integer division (clockTimestamp * 10000000 / time_scale),
+		// which for time_scale values that don't evenly divide 10000000 (e.g. 60) produces a systematic
+		// downward bias that manifests as near-1.0 subframe values on the resulting FTimecode.
+		const double TotalSeconds = static_cast<double>(InMPEGTimecode->clockTimestamp) / static_cast<double>(InMPEGTimecode->time_scale);
 		OutFramerate = FFrameRate(InMPEGTimecode->time_scale, InMPEGTimecode->num_units_in_tick);
 		const FFrameTime FrameTime = OutFramerate.GetValue().AsFrameTime(TotalSeconds);
 		const bool bDropFrame = InMPEGTimecode->ct_type > 1;
